Move UnsyncUI code into Engine/Source/Programs

diff --git a/Engine/Source/Programs/UnsyncUI/UnsyncUI/App.xaml b/Engine/Source/Programs/UnsyncUI/UnsyncUI/App.xaml
new file mode 100644
--- /dev/null
+++ b/Engine/Source/Programs/UnsyncUI/UnsyncUI/App.xaml
@@ -0,0 +1,40 @@
+﻿<!--Copyright Epic Games, Inc. All Rights Reserved.-->
+
+<Application x:Class="UnsyncUI.App"
+             xmlns="http://schemas.microsoft.com/winfx/2006/xaml/presentation"
+             xmlns:x="http://schemas.microsoft.com/winfx/2006/xaml"
+			 xmlns:s="clr-namespace:System;assembly=mscorlib"
+             xmlns:local="clr-namespace:UnsyncUI"
+             StartupUri="MainWindow.xaml">
+    <Application.Resources>
+        <Style TargetType="local:MainWindow">
+            <Setter Property="Margin" Value="4" />
+        </Style>
+        <Style TargetType="Button">
+            <Setter Property="Margin" Value="4" />
+        </Style>
+        <Style TargetType="TextBox">
+            <Setter Property="Margin" Value="4" />
+        </Style>
+		<Style TargetType="ProgressBar">
+			<Setter Property="Height" Value="4" />
+			<Setter Property="Margin" Value="4" />
+		</Style>
+		
+		<Visibility x:Key="Hidden">Hidden</Visibility>
+		<Visibility x:Key="Collapsed">Collapsed</Visibility>
+		<Visibility x:Key="Visible">Visible</Visibility>
+		
+		<s:Boolean x:Key="True">True</s:Boolean>
+		<s:Boolean x:Key="False">False</s:Boolean>
+		
+		<local:AllConverter x:Key="BoolConv_All" TrueValue="{StaticResource True}" FalseValue="{StaticResource False}"/>
+
+		<local:AllConverter x:Key="VisConv_All_Collapsed" TrueValue="{StaticResource Visible}" FalseValue="{StaticResource Collapsed}"/>
+		<local:NotConverter x:Key="NotConv"/>
+		<local:VisibilityConverter x:Key="VisConv_Collapsed" FalseState="Collapsed" />
+		<local:VisibilityConverter x:Key="VisConv_Collapsed_Inv" TrueState="Collapsed" FalseState="Visible" />
+		<local:VisibilityConverter x:Key="VisConv_Hidden" FalseState="Hidden" />
+		
+	</Application.Resources>
+</Application>
diff --git a/Engine/Source/Programs/UnsyncUI/UnsyncUI/App.xaml.cs b/Engine/Source/Programs/UnsyncUI/UnsyncUI/App.xaml.cs
new file mode 100644
--- /dev/null
+++ b/Engine/Source/Programs/UnsyncUI/UnsyncUI/App.xaml.cs
@@ -0,0 +1,149 @@
+// Copyright Epic Games, Inc. All Rights Reserved.
+
+using System;
+using System.Collections.Generic;
+using System.IO;
+using System.Text.Json;
+using System.Windows;
+
+namespace UnsyncUI
+{
+	/// <summary>
+	/// Interaction logic for App.xaml
+	/// </summary>
+	public partial class App : Application
+    {
+		public new static App Current => Application.Current as App;
+
+		public Config Config { get; private set; }
+		public string UnsyncPath { get; private set; }
+		public UserPreferences UserConfig { get; private set; }
+		public string DefaultSearchTerms { get; private set; } = "";
+
+		protected override void OnStartup(StartupEventArgs e)
+		{
+			string configFile = e.Args.Length > 0 ? e.Args[0] : "unsyncui.xml";
+
+			for (int i=1; i<e.Args.Length; ++i)
+			{
+				bool isLast = i + 1 == e.Args.Length;
+				var arg = e.Args[i];
+				if (arg == "--search" && !isLast)
+				{
+					DefaultSearchTerms = e.Args[i + 1];
+					++i;
+				}
+			}
+
+			var oldWorkingDir = Environment.CurrentDirectory;
+
+			if (File.Exists(configFile))
+			{
+				// Move the working dir into the same as the config file.
+				// This allows relative paths to automatically be relative to the config.
+				var workingDir = Path.GetDirectoryName(Path.GetFullPath(configFile));
+				Environment.CurrentDirectory = workingDir;
+
+				try
+				{
+					Config = new Config(configFile);
+				}
+				catch (Exception ex)
+				{
+					MessageBox.Show($"Failed to load configuration from \"{configFile}\". {ex}");
+				}
+			}
+			else if (e.Args.Length > 0)
+			{
+				MessageBox.Show($"The configuration file \"{configFile}\" does not exist. Projects will not be available.");
+			}
+
+			string toolName = "unsync.exe";
+
+			// If config specifies the explicit tool path, then use that.
+			// Otherwise try same directory as config file.
+			// Finally, try to look in the original working directory.
+
+			string newWorkingDirTool = Path.GetFullPath(toolName);
+			string oldWorkingDirTool = Path.GetFullPath(Path.Combine(oldWorkingDir, toolName));
+
+			if (Config != null && Config.UnsyncPath != null)
+			{
+				UnsyncPath = Config.UnsyncPath;
+			}
+			else if (File.Exists(newWorkingDirTool))
+			{
+				UnsyncPath = newWorkingDirTool;
+			}
+			else if (File.Exists(oldWorkingDirTool))
+			{
+				UnsyncPath = oldWorkingDirTool;
+			}
+
+			if (UnsyncPath == null || !File.Exists(UnsyncPath))
+			{
+				string name = UnsyncPath ?? toolName;
+				MessageBox.Show($"Failed to find \"{name}\".");
+				Shutdown();
+			}
+
+			UserConfig = UserPreferences.Load();
+			if (UserConfig.DFS == "")
+			{
+				UserConfig.DFS = null;
+			}
+
+			base.OnStartup(e);
+		}
+
+		protected override void OnExit(ExitEventArgs e)
+		{
+			UserConfig.Save();
+			base.OnExit(e);
+		}
+	}
+
+	public sealed class UserPreferences
+	{
+		// Change this guid if any changes are made which break compat with existing user preferences files
+		private static readonly Guid version = Guid.Parse("1386020B-374C-43CE-8FBF-74AAE69EEFD3");
+		private static readonly string userDir = Path.Combine(Environment.GetFolderPath(Environment.SpecialFolder.LocalApplicationData), "UnsyncUI");
+		private static readonly string userFile = Path.Combine(userDir, $"user_{version}.json");
+
+		public string Proxy { get; set; }
+		public string DFS { get; set; }
+		public Dictionary<string, string> ProjectDestinationMap { get; set; } = new Dictionary<string, string>();
+		public string CustomSrcPath { get; set; }
+		public string CustomDstPath { get; set; }
+		public string AdditionalArgs { get; set; }
+
+		public static UserPreferences Load()
+		{
+			try
+			{
+				if (File.Exists(userFile))
+				{
+					return JsonSerializer.Deserialize<UserPreferences>(File.ReadAllText(userFile));
+				}
+			}
+			catch (Exception ex)
+			{
+				MessageBox.Show($"Failed to load user preferences from \"{userFile}\". {ex}");
+			}
+
+			return new UserPreferences();
+		}
+
+		public void Save()
+		{
+			try
+			{
+				if (!Directory.Exists(userDir))
+					Directory.CreateDirectory(userDir);
+
+				File.WriteAllText(userFile, JsonSerializer.Serialize(this, new JsonSerializerOptions() { WriteIndented = true }));
+			}
+			catch { }
+		}
+	}
+}
diff --git a/Engine/Source/Programs/UnsyncUI/UnsyncUI/AssemblyInfo.cs b/Engine/Source/Programs/UnsyncUI/UnsyncUI/AssemblyInfo.cs
new file mode 100644
--- /dev/null
+++ b/Engine/Source/Programs/UnsyncUI/UnsyncUI/AssemblyInfo.cs
@@ -0,0 +1,12 @@
+// Copyright Epic Games, Inc. All Rights Reserved.
+
+using System.Windows;
+
+[assembly: ThemeInfo(
+    ResourceDictionaryLocation.None, //where theme specific resource dictionaries are located
+                                     //(used if a resource is not found in the page,
+                                     // or application resource dictionaries)
+    ResourceDictionaryLocation.SourceAssembly //where the generic resource dictionary is located
+                                              //(used if a resource is not found in the page,
+                                              // app, or any theme specific resource dictionaries)
+)]
diff --git a/Engine/Source/Programs/UnsyncUI/UnsyncUI/BuildSelectorControl.xaml b/Engine/Source/Programs/UnsyncUI/UnsyncUI/BuildSelectorControl.xaml
new file mode 100644
--- /dev/null
+++ b/Engine/Source/Programs/UnsyncUI/UnsyncUI/BuildSelectorControl.xaml
@@ -0,0 +1,166 @@
+<!--Copyright Epic Games, Inc. All Rights Reserved.-->
+
+<UserControl x:Class="UnsyncUI.BuildSelectorControl"
+             xmlns="http://schemas.microsoft.com/winfx/2006/xaml/presentation"
+             xmlns:x="http://schemas.microsoft.com/winfx/2006/xaml"
+             xmlns:mc="http://schemas.openxmlformats.org/markup-compatibility/2006" 
+             xmlns:d="http://schemas.microsoft.com/expression/blend/2008" 
+             xmlns:local="clr-namespace:UnsyncUI"
+			 xmlns:dat="clr-namespace:System.ComponentModel;assembly=WindowsBase"
+             mc:Ignorable="d" 
+             d:DesignHeight="450" d:DesignWidth="800">
+
+	<Grid>
+		<Grid.RowDefinitions>
+			<RowDefinition Height="*" />
+			<RowDefinition Height="auto" />
+			<RowDefinition Height="auto" />
+			<RowDefinition Height="auto" />
+		</Grid.RowDefinitions>
+
+		<!-- Builds Groupbox -->
+		<GroupBox Grid.Row="0">
+
+			<GroupBox.Header>
+				<Grid>
+					<TextBlock Margin="4" Visibility="{Binding Path=OnRefreshBuildsClicked.Enabled, Converter={StaticResource VisConv_Collapsed}}">
+					<Run Text="Builds -"/>
+					<Hyperlink Command="{Binding OnRefreshBuildsClicked}">Refresh</Hyperlink>
+					</TextBlock>
+
+					<TextBlock Margin="4" Visibility="{Binding Path=OnStopRefreshBuildsClicked.Enabled, Converter={StaticResource VisConv_Collapsed}}">
+					<Run Text="Builds -"/>
+					<Hyperlink Command="{Binding OnStopRefreshBuildsClicked}">Stop</Hyperlink>
+					</TextBlock>
+				</Grid>
+			</GroupBox.Header>
+
+			<Grid>
+				<Grid.RowDefinitions>
+					<RowDefinition Height="auto"/>
+					<RowDefinition Height="auto"/>
+					<RowDefinition Height="*"/>
+				</Grid.RowDefinitions>
+
+				<ProgressBar Grid.Row="0" IsIndeterminate="True" Visibility="{Binding Path=OnStopRefreshBuildsClicked.Enabled, Converter={StaticResource VisConv_Hidden}}" Margin="4,2,4,2"/>
+
+				<Grid Grid.Row="1">
+					<Grid.ColumnDefinitions>
+						<ColumnDefinition Width="*"/>
+						<ColumnDefinition Width="auto"/>
+					</Grid.ColumnDefinitions>
+					<TextBox Grid.Row ="0" Grid.Column="0" Text="{Binding SearchTerms, UpdateSourceTrigger=PropertyChanged}" />
+					<TextBlock Grid.Row="0" Grid.Column="1" Text="Search" Margin="8,4,8,4"/>
+				</Grid>
+
+				<ListView Grid.Row="2" SelectionMode="Single" ItemsSource="{Binding Builds}" SelectedItem="{Binding SelectedBuild}" BorderBrush="Transparent" Name="SelectedBuild" GridViewColumnHeader.Click="BuildsHeader_Click">
+					<ListView.ContextMenu>
+						<ContextMenu>
+							<MenuItem Header="Copy Build CL" Command="{Binding CopyBuildCLClicked}" />
+							<MenuItem Header="Open Build Location" Command="{Binding OpenBuildLocationClicked}" />
+						</ContextMenu>
+					</ListView.ContextMenu>
+					<ListView.View>
+						<GridView>
+							<GridViewColumn x:Name="colStream" Width="200" Header="Stream" DisplayMemberBinding="{Binding Stream}" />
+							<GridViewColumn x:Name="colCL" Width="100" Header="CL" DisplayMemberBinding="{Binding CL}" />
+							<GridViewColumn x:Name="colSuffix" Width="100" Header="Suffix" DisplayMemberBinding="{Binding Suffix}" />
+						</GridView>
+					</ListView.View>
+				</ListView>
+			</Grid>
+		</GroupBox>
+
+		<!-- Platforms Groupbox -->
+		<GroupBox Grid.Row="1" MinHeight="62" DataContext="{Binding SelectedItem, ElementName=SelectedBuild}" Visibility="{Binding SelectedItem, ElementName=SelectedBuild, Converter={StaticResource VisConv_Collapsed}}">
+			
+			<GroupBox.Header>
+				<Grid>
+					<TextBlock Margin="4" Visibility="{Binding Path=OnRefreshPlatformsClicked.Enabled, Converter={StaticResource VisConv_Collapsed}}">
+					<Run Text="Platforms -"/>
+					<Hyperlink Command="{Binding OnRefreshPlatformsClicked}">Refresh</Hyperlink>
+					</TextBlock>
+
+					<TextBlock Margin="4" Visibility="{Binding Path=OnStopRefreshPlatformsClicked.Enabled, Converter={StaticResource VisConv_Collapsed}}">
+					<Run Text="Platforms -"/>
+					<Hyperlink Command="{Binding OnStopRefreshPlatformsClicked}">Stop</Hyperlink>
+					</TextBlock>
+				</Grid>
+			</GroupBox.Header>
+			
+			<Grid>
+				<Grid.RowDefinitions>
+					<RowDefinition Height="auto"/>
+					<RowDefinition Height="*"/>
+				</Grid.RowDefinitions>
+
+				<ProgressBar Grid.Row="0" IsIndeterminate="True" Visibility="{Binding Path=OnStopRefreshPlatformsClicked.Enabled, Converter={StaticResource VisConv_Hidden}}" Margin="4,2,4,2"/>
+
+				<TextBlock Grid.Row="0" Grid.RowSpan="2" Text="No platforms found for this build" VerticalAlignment="Center" HorizontalAlignment="Center">
+					<TextBlock.Visibility>
+						<MultiBinding Converter="{StaticResource VisConv_All_Collapsed}">
+							<Binding Path="Platforms.Count" Converter="{StaticResource NotConv}"/>
+							<Binding Path="OnRefreshPlatformsClicked.Enabled"/>
+						</MultiBinding>
+					</TextBlock.Visibility>
+				</TextBlock>
+
+				<ListView Grid.Row="1" Foreground="Transparent" Background="Transparent" BorderBrush="Transparent" ItemsSource="{Binding Platforms}" Name="SelectedPlatforms">
+					<ListView.ItemTemplate>
+						<DataTemplate>
+							<CheckBox Margin="4,0,4,0" Content="{Binding Platform}" IsChecked="{Binding Path=IsSelected, RelativeSource={RelativeSource FindAncestor, AncestorType={x:Type ListViewItem}}}"/>
+						</DataTemplate>
+					</ListView.ItemTemplate>
+					<ListView.ItemsPanel>
+						<ItemsPanelTemplate>
+							<WrapPanel Orientation="Horizontal" 
+									   Width="{Binding RelativeSource={RelativeSource FindAncestor,AncestorType={x:Type ScrollViewer}},Path=ActualWidth}"/>
+						</ItemsPanelTemplate>
+					</ListView.ItemsPanel>
+					<ListView.ItemContainerStyle>
+						<Style TargetType="ListViewItem">
+							<Setter Property="Template">
+								<Setter.Value>
+									<ControlTemplate TargetType="ListViewItem">
+										<ContentPresenter/>
+									</ControlTemplate>
+								</Setter.Value>
+							</Setter>
+						</Style>
+					</ListView.ItemContainerStyle>
+				</ListView>
+			</Grid>
+			
+		</GroupBox>
+
+		<!-- Destination Path -->
+		<GroupBox Grid.Row="2">
+			<GroupBox.Header>
+				<Grid>
+					<TextBlock Margin="4" Visibility="{Binding Path=!ResetDestinationClicked.Enabled, Converter={StaticResource VisConv_Collapsed}}">
+						<Run Text="Destination"/>
+					</TextBlock>
+					<TextBlock Margin="4" Visibility="{Binding Path=ResetDestinationClicked.Enabled, Converter={StaticResource VisConv_Collapsed}}">
+						<Run Text="Destination - "/>
+					<Hyperlink Command="{Binding ResetDestinationClicked}">Reset</Hyperlink>
+					</TextBlock>
+				</Grid>
+			</GroupBox.Header>
+
+			<local:SelectFolderControl SelectedPath="{Binding DstPath, UpdateSourceTrigger=PropertyChanged}" Description="Select a destination directory to sync builds to."/>
+		</GroupBox>
+
+		<!-- Start sync -->
+		<TextBlock Grid.Row="3" Margin="4" HorizontalAlignment="Center">
+			<Hyperlink Command="{Binding OnStartSync}" CommandParameter="{Binding SelectedItems, ElementName=SelectedPlatforms}">
+				<Hyperlink.IsEnabled>
+					<MultiBinding Converter="{StaticResource BoolConv_All}">
+						<Binding ElementName="SelectedPlatforms" Path="SelectedItems.Count"/>
+						<Binding Path="DstPath.Length"/>
+					</MultiBinding>
+				</Hyperlink.IsEnabled>
+				Start Sync
+			</Hyperlink>
+		</TextBlock>
+	</Grid>
+</UserControl>
diff --git a/Engine/Source/Programs/UnsyncUI/UnsyncUI/BuildSelectorControl.xaml.cs b/Engine/Source/Programs/UnsyncUI/UnsyncUI/BuildSelectorControl.xaml.cs
new file mode 100644
--- /dev/null
+++ b/Engine/Source/Programs/UnsyncUI/UnsyncUI/BuildSelectorControl.xaml.cs
@@ -0,0 +1,435 @@
+// Copyright Epic Games, Inc. All Rights Reserved.
+
+using System;
+using System.Collections;
+using System.Collections.Generic;
+using System.Collections.ObjectModel;
+using System.ComponentModel;
+using System.Diagnostics;
+using System.IO;
+using System.Linq;
+using System.Threading;
+using System.Threading.Tasks;
+using System.Threading.Tasks.Dataflow;
+using System.Windows;
+using System.Windows.Controls;
+using System.Windows.Data;
+
+namespace UnsyncUI
+{
+	/// <summary>
+	/// Interaction logic for BuildSelectorControl.xaml
+	/// </summary>
+	public partial class BuildSelectorControl : UserControl
+	{
+		GridViewColumnHeader lastHeaderClicked = null;
+		ListSortDirection lastDirection = ListSortDirection.Ascending;
+
+		public BuildSelectorControl()
+		{
+			InitializeComponent();
+		}
+
+		private void SyncHyperlink_Click(object sender, RoutedEventArgs e)
+		{
+			(DataContext as ProjectModel).StartSync(SelectedPlatforms.SelectedItems.OfType<BuildPlatformModel>());
+		}
+
+		private void BuildsHeader_Click(object sender, RoutedEventArgs e)
+		{
+			var headerClicked = e.OriginalSource as GridViewColumnHeader;
+			if (headerClicked == null || headerClicked.Role == GridViewColumnHeaderRole.Padding)
+				return;
+
+			ListSortDirection direction;
+			if (headerClicked != lastHeaderClicked)
+			{
+				direction = ListSortDirection.Ascending;
+			}
+			else
+			{
+				direction = lastDirection == ListSortDirection.Ascending
+					? ListSortDirection.Descending
+					: ListSortDirection.Ascending;
+			}
+
+			if (headerClicked.Column == colStream)
+			{
+				Sort((colStream.Header as string, direction), (colCL.Header as string, ListSortDirection.Descending));
+			}
+			else
+			{
+				var columnBinding = headerClicked.Column.DisplayMemberBinding as Binding;
+				var sortBy = columnBinding?.Path.Path ?? headerClicked.Column.Header as string;
+
+				Sort((sortBy, direction));
+			}
+
+			lastHeaderClicked = headerClicked;
+			lastDirection = direction;
+		}
+
+		private void Sort(params (string field, ListSortDirection direction)[] args)
+		{
+			var dataView = CollectionViewSource.GetDefaultView(SelectedBuild.ItemsSource);
+			dataView.SortDescriptions.Clear();
+			foreach (var desc in args)
+			{
+				dataView.SortDescriptions.Add(new SortDescription(desc.field, desc.direction));
+			}
+			dataView.Refresh();
+		}
+	}
+
+	public struct Build
+	{
+		public string Stream;
+		public string Platform;
+		public string CL;
+		public string Suffix;
+		public string FullPath;
+	}
+
+	public sealed class BuildPlatformModel
+	{
+
+		public BuildModel Build { get; }
+		public string Platform { get; }
+		public string FullPath { get; }
+
+		public string Name => Build == null
+			? "Custom Sync" 
+			: $"{Build.Stream} - {Platform} - CL {Build.CL}";
+
+		public string SafeName => Name.Replace('/', '+').Replace(" ", "");
+
+		public BuildPlatformModel(BuildModel build, string platform, string fullPath)
+		{
+			Build = build;
+			Platform = platform;
+			FullPath = fullPath;
+		}
+	}
+
+	public sealed class BuildModel : BaseModel
+	{
+		private CancellationTokenSource cts;
+
+		public string Stream => Template.Stream;
+		public string CL => Template.CL;
+		public string Suffix => Template.Suffix;
+
+		public Config.Directory RootDir { get; }
+		public string Path { get; }
+
+		public Config.BuildTemplate Template { get; }
+
+		public Command OnRefreshPlatformsClicked { get; }
+		public Command OnStopRefreshPlatformsClicked { get; }
+
+		public ObservableCollection<BuildPlatformModel> Platforms { get; } = new ObservableCollection<BuildPlatformModel>();
+
+		public BuildModel(string path, Config.Directory rootDir, Config.BuildTemplate template)
+		{
+			Path = path;
+			RootDir = rootDir;
+			Template = template;
+
+			OnRefreshPlatformsClicked = new Command(() => PopulatePlatforms()) { Enabled = true };
+			OnStopRefreshPlatformsClicked = new Command(() =>
+			{
+				cts?.Cancel();
+			});
+		}
+
+		public bool HasPlatform(string platform) => Platforms.Where(p => p.Platform == platform).Any();
+
+		private async Task EnumeratePlatforms(Config.Directory currentDir, string path, Config.BuildTemplate template)
+		{
+			if (!currentDir.Parse(path, ref template))
+				return;
+
+			foreach (var childDir in await AsyncIO.EnumerateDirectoriesAsync(path, cts.Token))
+			{
+				if (System.IO.Path.GetFileName(childDir) == ".unsync")
+				{
+					// This folder is a valid build
+					Platforms.Add(new BuildPlatformModel(this, template.Platform, path));
+				}
+				else
+				{
+					foreach (var subDir in currentDir.SubDirectories)
+					{
+						await EnumeratePlatforms(subDir, childDir, template);
+					}
+				}
+			}
+		}
+
+		public async void PopulatePlatforms()
+		{
+			cts = new CancellationTokenSource();
+
+			OnRefreshPlatformsClicked.Enabled = false;
+			OnStopRefreshPlatformsClicked.Enabled = true;
+
+			Platforms.Clear();
+
+			try
+			{
+				await EnumeratePlatforms(RootDir, Path, Template);
+			}
+			catch (OperationCanceledException)
+			{ }
+			finally
+			{
+				OnRefreshPlatformsClicked.Enabled = true;
+				OnStopRefreshPlatformsClicked.Enabled = false;
+			}
+		}
+	}
+
+	public sealed class ProjectModel : TabModel
+	{
+		private bool firstRefresh = true;
+		private Action<IEnumerable<(string DstPath, string[] Exclusions, BuildPlatformModel Model)>> onBuildsSelected;
+
+		public Config.Project Definition { get; }
+		public string Name => Definition.Name;
+
+		public Command OnRefreshBuildsClicked { get; }
+		public Command OnStopRefreshBuildsClicked { get; }
+		public Command<IList> OnStartSync { get; }
+		public Command CopyBuildCLClicked { get; }
+		public Command OpenBuildLocationClicked { get; }
+		public Command ResetDestinationClicked { get; }
+
+		public ObservableCollection<BuildModel> Builds { get; } = new ObservableCollection<BuildModel>();
+		public ICollectionView BuildsView { get; }
+
+		public bool BuildsViewFilter(object item)
+		{
+			BuildModel build = (BuildModel)item;
+
+			foreach (var term in searchTerms.Split(' '))
+			{
+				if (!build.CL.Contains(term, StringComparison.InvariantCultureIgnoreCase)
+				&& !build.Stream.Contains(term, StringComparison.InvariantCultureIgnoreCase)
+				&& !build.Suffix.Contains(term, StringComparison.InvariantCultureIgnoreCase))
+				{
+					return false;
+				}
+			}
+
+			return true;
+		}
+
+		private BuildModel selectedBuild = null;
+		public BuildModel SelectedBuild
+		{
+			get => selectedBuild;
+			set
+			{
+				SetProperty(ref selectedBuild, value);
+
+				if (selectedBuild != null)
+				{
+					selectedBuild.PopulatePlatforms();
+				}
+			}
+		}
+
+		private string searchTerms = "";
+		public string SearchTerms
+		{
+			get => searchTerms;
+			set
+			{
+				SetProperty(ref searchTerms, value);
+				BuildsView.Refresh();
+			}
+		}
+
+		private CancellationTokenSource cts = null;
+
+		/// Default download directory (either auto-detected or loaded from project config)
+		private string defaultDstPath = null;
+
+		/// Current download path selected by user or loaded from user preferences.
+		/// This value persists between application sessions (saved per project).
+		/// It can be reset to default by pressing "Reset" button in the UI.
+		private string dstPath = null;
+
+		public string DstPath
+		{
+			get => dstPath;
+			set
+			{
+				SetProperty(ref dstPath, value);
+				App.Current.UserConfig.ProjectDestinationMap[Name] = value;
+
+				ResetDestinationClicked.Enabled = CanResetDestinationToDefault;
+			}
+		}
+
+		public ProjectModel(Config.Project definition, Action<IEnumerable<(string DstPath, string[] Exclusions, BuildPlatformModel Model)>> onBuildsSelected)
+		{
+			this.Definition = definition;
+			this.onBuildsSelected = onBuildsSelected;
+			OnRefreshBuildsClicked = new Command(RefreshBuildsAsync) { Enabled = true };
+			OnStopRefreshBuildsClicked = new Command(StopRefreshBuilds) { Enabled = false };
+			OnStartSync = new Command<IList>(l => StartSync(l.OfType<BuildPlatformModel>())) { Enabled = true };
+			CopyBuildCLClicked = new Command(CopyBuildCL) { Enabled = true };
+			OpenBuildLocationClicked = new Command(OpenBuildLocation) { Enabled = true };
+
+			BuildsView = CollectionViewSource.GetDefaultView(Builds);
+			BuildsView.Filter = BuildsViewFilter;
+			BuildsView.SortDescriptions.Clear();
+			BuildsView.SortDescriptions.Add(new SortDescription("CL", ListSortDirection.Descending));
+
+			SearchTerms = App.Current.DefaultSearchTerms;
+
+			if (definition.Destination != null)
+			{
+				defaultDstPath = Path.GetFullPath(definition.Destination);
+			}
+			else
+			{
+				// If this is the first time running, choose the destination dir as the largest local disk
+				// @todo: remember user choice
+				var bestDrive = DriveInfo.GetDrives()
+				.Where(d => d.DriveType == DriveType.Fixed)
+				.Where(d => !d.VolumeLabel.Contains("google", StringComparison.InvariantCultureIgnoreCase)) // Ignore Google Drive (it appears as a 'fixed' disk)
+				.OrderByDescending(d => d.TotalSize).FirstOrDefault();
+				if (bestDrive != null)
+				{
+					defaultDstPath = Path.Combine(bestDrive.RootDirectory.FullName, "Unsync", definition.Name);
+				}
+			}
+
+			// If we failed to find a reasonable default destination then set it to a non-null empty string,
+			// which prevents sync command until user explicitly sets the path
+			if (defaultDstPath == null)
+			{
+				defaultDstPath = "";
+			}
+
+			string userDstPath = null;
+
+			if (App.Current.UserConfig.ProjectDestinationMap.TryGetValue(Name, out userDstPath))
+			{
+				dstPath = userDstPath;
+			}
+			else
+			{
+				dstPath = defaultDstPath;
+			}
+
+			ResetDestinationClicked = new Command(ResetDestinationToDefault) { Enabled = CanResetDestinationToDefault };
+		}
+		private bool CanResetDestinationToDefault => dstPath != defaultDstPath && defaultDstPath.Length != 0;
+
+		private void ResetDestinationToDefault()
+		{
+			DstPath = defaultDstPath;
+		}
+
+		private void StopRefreshBuilds()
+		{
+			if (cts != null)
+			{
+				cts.Cancel();
+				OnStopRefreshBuildsClicked.Enabled = false;
+			}
+		}
+
+		public async void RefreshBuildsAsync()
+		{
+			cts = new CancellationTokenSource();
+
+			OnRefreshBuildsClicked.Enabled = false;
+			OnStopRefreshBuildsClicked.Enabled = true;
+
+			Builds.Clear();
+			SelectedBuild = null;
+
+			try
+			{
+				var timer = new Stopwatch();
+				timer.Start();
+				var buildsPipe = Definition.EnumerateBuilds(cts.Token);
+
+				while (await buildsPipe.OutputAvailableAsync(cts.Token))
+				{
+					Builds.Add(await buildsPipe.ReceiveAsync());
+				}
+
+				timer.Stop();
+				Debug.WriteLine($"Search took {timer.Elapsed.TotalSeconds} seconds.");
+			}
+			catch (IOException)
+			{
+				// @todo show an error message
+			}
+			catch (OperationCanceledException)
+			{ }
+			finally
+			{
+				OnRefreshBuildsClicked.Enabled = true;
+				OnStopRefreshBuildsClicked.Enabled = false;
+			}
+		}
+
+		public void StartSync(IEnumerable<BuildPlatformModel> selectedBuilds)
+		{
+			onBuildsSelected?.Invoke(selectedBuilds.Select(s => (Path.Combine(dstPath, s.Platform), Definition.Exclusions?.ToArray(), s)).ToList());
+		}
+
+		public override void OnSelected()
+		{
+			if (firstRefresh)
+			{
+				RefreshBuildsAsync();
+				firstRefresh = false;
+			}
+		}
+
+		private void CopyBuildCL()
+		{
+			if (selectedBuild != null)
+			{
+				// Clipboard can be locked if timing is just right(wrong), retry a few times on failure.
+                int retryDelay = 0;
+				const int attempts = 7;
+				for (int attempt = 0; attempt < attempts; attempt++)
+				{
+					if (retryDelay > 0)
+					{
+                        System.Threading.Thread.Sleep(retryDelay);
+                        retryDelay *= 2;
+                    }
+					else
+					{
+						retryDelay = 5;
+                    }
+                    try
+                    {
+                        Clipboard.SetText(selectedBuild.CL, TextDataFormat.Text);
+						break;
+                    }
+                    catch (System.Runtime.InteropServices.COMException)
+                    {
+	                }
+                }
+			}
+		}
+
+		private void OpenBuildLocation()
+		{
+			if (selectedBuild != null)
+			{
+				Shell.LaunchExplorer(selectedBuild.Path);
+			}
+		}
+	}
+}
diff --git a/Engine/Source/Programs/UnsyncUI/UnsyncUI/Config.cs b/Engine/Source/Programs/UnsyncUI/UnsyncUI/Config.cs
new file mode 100644
--- /dev/null
+++ b/Engine/Source/Programs/UnsyncUI/UnsyncUI/Config.cs
@@ -0,0 +1,178 @@
+// Copyright Epic Games, Inc. All Rights Reserved.
+
+using System;
+using System.Collections.Generic;
+using System.IO;
+using System.Linq;
+using System.Text.RegularExpressions;
+using System.Threading;
+using System.Threading.Tasks;
+using System.Threading.Tasks.Dataflow;
+using System.Xml.Linq;
+
+namespace UnsyncUI
+{
+	public sealed class Config
+	{
+		public sealed class Proxy
+		{
+			public string Name { get; set; }
+			public string Path { get; set; }
+		}
+
+		public struct BuildTemplate
+		{
+			public string Stream;
+			public string CL;
+			public string Suffix;
+			public string Platform;
+
+			public bool IsStreamClFound => Stream != null && CL != null;
+		}
+
+		public sealed class Directory
+		{
+			public Regex Regex { get; }
+			public string Stream { get; }
+			public string Suffix { get; }
+			public string CL { get; }
+			public string Platform { get; }
+			public List<Directory> SubDirectories { get; }
+
+			public Directory(XElement node)
+			{
+				Regex = new Regex($@"^{node.Attribute("regex")?.Value}$", RegexOptions.IgnoreCase);
+				Stream = node.Attribute("stream")?.Value;
+				Suffix = node.Attribute("suffix")?.Value;
+				CL = node.Attribute("cl")?.Value;
+				Platform = node.Attribute("platform")?.Value;
+				SubDirectories = node.Elements("dir").Select(d => new Directory(d)).ToList();
+			}
+
+			public bool Parse(string path, ref BuildTemplate template)
+			{
+				var match = Regex.Match(Path.GetFileName(path));
+				if (!match.Success)
+					return false;
+
+				string RegexReplace(string field) => Regex.Replace(field, @"\$([0-9]+)", m => match.Groups[int.Parse(m.Groups[1].Value)].Value);
+
+				if (Stream != null)
+				{
+					template.Stream = RegexReplace(Stream);
+				}
+				if (CL != null)
+				{
+					template.CL = RegexReplace(CL);
+				}
+				if (Suffix != null)
+				{
+					template.Suffix = RegexReplace(Suffix);
+				}
+				if (Platform != null)
+				{
+					template.Platform = RegexReplace(Platform);
+				}
+
+				return true;
+			}
+		}
+
+		public sealed class Project
+		{
+			public string Name { get; set; }
+			public string Root { get; set; }
+			public string Destination { get; set; }
+			public List<string> Exclusions { get; set; }
+			public List<Directory> Children { get; set; }
+
+			private Task EnumerateBuilds(ITargetBlock<BuildModel> pipe, Directory d, string path, BuildTemplate template, CancellationToken cancelToken)
+			{
+				return Task.Run(async () =>
+				{
+					if (!d.Parse(path, ref template))
+						return;
+
+					if (template.IsStreamClFound)
+					{
+						pipe.Post(new BuildModel(path, d, template));
+					}
+					else
+					{
+						var jobs = new List<Task>();
+						foreach (var childDir in await AsyncIO.EnumerateDirectoriesAsync(path, cancelToken))
+						{
+							jobs.AddRange(d.SubDirectories.Select(s => EnumerateBuilds(pipe, s, childDir, template, cancelToken)));
+						}
+
+						await Task.WhenAll(jobs);
+					}
+				});
+			}
+
+			public ISourceBlock<BuildModel> EnumerateBuilds(CancellationToken cancelToken)
+			{
+				var pipe = new BufferBlock<BuildModel>();
+				Task.Run(async () =>
+				{
+					try
+					{
+						var allTasks = new List<Task>();
+						foreach (var dir in await AsyncIO.EnumerateDirectoriesAsync(Root, cancelToken))
+						{
+							allTasks.AddRange(Children.Select(c => EnumerateBuilds(pipe, c, dir, default(BuildTemplate), cancelToken)));
+						}
+
+						await Task.WhenAll(allTasks);
+					}
+					finally
+					{
+						pipe.Complete();
+					}
+				});
+
+				return pipe;
+			}
+		}
+
+		public List<Proxy> Proxies { get; set; } = new List<Proxy>();
+		public List<Project> Projects { get; set; }
+
+		public string UnsyncPath { get; set; }
+		public string DFS { get; set; }
+
+		public Config(string filename)
+		{
+			var rootNode = XDocument.Load(filename).Root;
+
+			UnsyncPath = rootNode.Attribute("path")?.Value;
+
+			if (UnsyncPath != null && !File.Exists(UnsyncPath))
+			{
+				throw new Exception("Unable to find unsync.exe binary specified in config file.");
+			}
+
+			DFS = rootNode.Attribute("dfs")?.Value;
+
+			Proxies.Add(new Proxy()
+			{
+				Name = "(none)",
+				Path = null
+			});
+			Proxies.AddRange(rootNode.Element("proxies").Elements("proxy").Select(p => new Proxy()
+			{
+				Name = p.Attribute("name")?.Value,
+				Path = p.Attribute("path")?.Value
+			}));
+
+			Projects = rootNode.Element("projects").Elements("project").Select(p => new Project()
+			{
+				Name = p.Attribute("name")?.Value,
+				Root = p.Attribute("root")?.Value,
+				Destination = p.Attribute("dest")?.Value,
+				Children = p.Elements("dir").Select(d => new Directory(d)).ToList(),
+				Exclusions = p.Elements("exclude").Select(d => d.Value).ToList()
+			}).ToList();
+		}
+	}
+}
diff --git a/Engine/Source/Programs/UnsyncUI/UnsyncUI/Converters.cs b/Engine/Source/Programs/UnsyncUI/UnsyncUI/Converters.cs
new file mode 100644
--- /dev/null
+++ b/Engine/Source/Programs/UnsyncUI/UnsyncUI/Converters.cs
@@ -0,0 +1,137 @@
+// Copyright Epic Games, Inc. All Rights Reserved.
+
+using System;
+using System.Collections.Generic;
+using System.Globalization;
+using System.Linq;
+using System.Windows;
+using System.Windows.Data;
+using System.Windows.Media;
+
+namespace UnsyncUI
+{
+	public abstract class BaseConverter
+	{
+		protected static bool AsBool(object value)
+		{
+			if (value is bool)
+			{
+				return (bool)value;
+			}
+			else if (value is int)
+			{
+				return (int)value != 0;
+			}
+			else
+			{
+				//throw new ArgumentException("Unimplemented type.", nameof(value));
+				return value != null;
+			}
+		}
+	}
+	public sealed class VisibilityConverter : IValueConverter
+	{
+		public Visibility TrueState { get; set; } = Visibility.Visible;
+		public Visibility FalseState { get; set; } = Visibility.Collapsed;
+
+		public object Convert(object value, Type targetType, object parameter, CultureInfo culture)
+		{
+			if (value is bool)
+			{
+				return (bool)value
+					? TrueState
+					: FalseState;
+			}
+			else if (value is int)
+			{
+				return (int)value != 0
+					? TrueState
+					: FalseState;
+			}
+			else if (value == null)
+			{
+				return FalseState;
+			}
+			else if (value is string)
+			{
+				return !string.IsNullOrEmpty((string)value)
+					? TrueState
+					: FalseState;
+			}
+			else
+			{
+				return TrueState;
+			}
+		}
+
+		public object ConvertBack(object value, Type targetType, object parameter, CultureInfo culture)
+			=> throw new NotImplementedException();
+	}
+
+	public sealed class NotConverter : BaseConverter, IValueConverter
+	{
+		public object Convert(object value, Type targetType, object parameter, CultureInfo culture)
+			=> !AsBool(value);
+
+		public object ConvertBack(object value, Type targetType, object parameter, CultureInfo culture)
+			=> throw new NotImplementedException();
+	}
+
+	public sealed class AllConverter : BaseConverter, IMultiValueConverter
+	{
+		public object TrueValue { get; set; } = true;
+		public object FalseValue { get; set; } = false;
+
+		public object Convert(object[] values, Type targetType, object parameter, CultureInfo culture)
+			=> values.All(o => AsBool(o)) ? TrueValue : FalseValue;
+
+		public object[] ConvertBack(object value, Type[] targetTypes, object parameter, CultureInfo culture)
+			=> throw new NotImplementedException();
+	}
+
+	public sealed class ObjectConverter : BaseConverter, IValueConverter
+	{
+		public object TrueValue { get; set; } = true;
+		public object FalseValue { get; set; } = false;
+
+		public object Convert(object value, Type targetType, object parameter, CultureInfo culture)
+			=> AsBool(value) ? TrueValue : FalseValue;
+
+		public object ConvertBack(object value, Type targetType, object parameter, CultureInfo culture)
+			=> value == TrueValue;
+	}
+
+	public sealed class ConsoleColorConverter : IValueConverter
+	{
+		private Dictionary<ConsoleColor, SolidColorBrush> brushes = new Dictionary<ConsoleColor, SolidColorBrush>()
+		{
+			[ConsoleColor.Black] = new SolidColorBrush(Color.FromRgb(0, 0, 0)),
+			[ConsoleColor.DarkBlue] = new SolidColorBrush(Color.FromRgb(0, 0, 128)),
+			[ConsoleColor.DarkGreen] = new SolidColorBrush(Color.FromRgb(0, 128, 0)),
+			[ConsoleColor.DarkCyan] = new SolidColorBrush(Color.FromRgb(0, 128, 128)),
+			[ConsoleColor.DarkRed] = new SolidColorBrush(Color.FromRgb(128, 0, 0)),
+			[ConsoleColor.DarkMagenta] = new SolidColorBrush(Color.FromRgb(128, 0, 128)),
+			[ConsoleColor.DarkYellow] = new SolidColorBrush(Color.FromRgb(128, 128, 0)),
+			[ConsoleColor.Gray] = new SolidColorBrush(Color.FromRgb(170, 170, 170)),
+			[ConsoleColor.DarkGray] = new SolidColorBrush(Color.FromRgb(85, 85, 85)),
+			[ConsoleColor.Blue] = new SolidColorBrush(Color.FromRgb(0, 0, 255)),
+			[ConsoleColor.Green] = new SolidColorBrush(Color.FromRgb(0, 255, 0)),
+			[ConsoleColor.Cyan] = new SolidColorBrush(Color.FromRgb(0, 255, 255)),
+			[ConsoleColor.Red] = new SolidColorBrush(Color.FromRgb(255, 0, 0)),
+			[ConsoleColor.Magenta] = new SolidColorBrush(Color.FromRgb(255, 0, 255)),
+			[ConsoleColor.Yellow] = new SolidColorBrush(Color.FromRgb(255, 255, 0)),
+			[ConsoleColor.White] = new SolidColorBrush(Color.FromRgb(255, 255, 255)),
+		};
+
+		public object Convert(object value, Type targetType, object parameter, CultureInfo culture)
+		{
+			if (brushes.TryGetValue((ConsoleColor)value, out var brush))
+				return brush;
+
+			return brushes[ConsoleColor.Black];
+		}
+
+		public object ConvertBack(object value, Type targetType, object parameter, CultureInfo culture)
+			=> throw new NotImplementedException();
+	}
+}
diff --git a/Engine/Source/Programs/UnsyncUI/UnsyncUI/JobControl.xaml b/Engine/Source/Programs/UnsyncUI/UnsyncUI/JobControl.xaml
new file mode 100644
--- /dev/null
+++ b/Engine/Source/Programs/UnsyncUI/UnsyncUI/JobControl.xaml
@@ -0,0 +1,100 @@
+﻿<!--Copyright Epic Games, Inc. All Rights Reserved.-->
+
+<UserControl x:Class="UnsyncUI.JobControl"
+             xmlns="http://schemas.microsoft.com/winfx/2006/xaml/presentation"
+             xmlns:x="http://schemas.microsoft.com/winfx/2006/xaml"
+             xmlns:mc="http://schemas.openxmlformats.org/markup-compatibility/2006" 
+             xmlns:d="http://schemas.microsoft.com/expression/blend/2008" 
+             xmlns:local="clr-namespace:UnsyncUI"
+             mc:Ignorable="d" 
+             d:DesignWidth="800">
+	<UserControl.Resources>
+		<local:ConsoleColorConverter x:Key="ConsoleColorConverter"/>
+	</UserControl.Resources>
+	<Grid>
+		<Grid.RowDefinitions>
+			<RowDefinition Height="auto"/>
+			<RowDefinition Height="*"/>
+		</Grid.RowDefinitions>
+		<Grid Grid.Row="0">
+			<Grid.ColumnDefinitions>
+				<ColumnDefinition Width="auto"/>
+				<ColumnDefinition Width="auto"/>
+				<ColumnDefinition Width="*"/>
+				<ColumnDefinition Width="100"/>
+				<ColumnDefinition Width="auto"/>
+			</Grid.ColumnDefinitions>
+			<Grid.RowDefinitions>
+				<RowDefinition/>
+				<RowDefinition/>
+			</Grid.RowDefinitions>
+			<Expander  Grid.Column="0" VerticalAlignment="Center" Name="FilesExpander" IsExpanded="{Binding DetailsExpanded}" />
+			<TextBlock Grid.Column="1" VerticalAlignment="Center" Text="{Binding Name}"/>
+			<TextBlock Grid.Column="2" VerticalAlignment="Center" HorizontalAlignment="Right" Text="{Binding StatusMessage}" Foreground="{Binding StatusColor, Converter={StaticResource ConsoleColorConverter}}"/>
+			<TextBlock Grid.Column="3" VerticalAlignment="Center" HorizontalAlignment="Right" Text="{Binding ProgressMessage}" Margin="4,4,8,4"/>
+			<TextBlock Grid.Column="4" VerticalAlignment="Center" Margin="4">
+                <TextBlock Visibility="{Binding ShowCancelHyperlink, Converter={StaticResource VisConv_Collapsed}}" >
+                    <Hyperlink Command="{Binding OnCancelClicked}">Cancel</Hyperlink>
+                </TextBlock>
+                <TextBlock Visibility="{Binding ShowClearHyperlink, Converter={StaticResource VisConv_Collapsed}}" >
+                    <Hyperlink Command="{Binding OnClearClicked}">Clear</Hyperlink>
+                </TextBlock>
+            </TextBlock>
+
+			<ProgressBar Grid.Row="1" Grid.Column="0" Grid.ColumnSpan="5"
+                         Value="{Binding TotalProgress, Mode=OneWay}" 
+                         IsIndeterminate="{Binding IsIndeterminate, Mode=OneWay}"
+                         Visibility="{Binding ShowProgressBar, Converter={StaticResource VisConv_Collapsed}}" />
+		</Grid>
+		<Grid Grid.Row="1" Visibility="{Binding ElementName=FilesExpander, Path=IsExpanded, Converter={StaticResource VisConv_Collapsed}}">
+			<Grid.ColumnDefinitions>
+				<ColumnDefinition Width="auto"/>
+				<ColumnDefinition Width="16"/>
+				<ColumnDefinition Width="*"/>
+			</Grid.ColumnDefinitions>
+			<Grid.RowDefinitions>
+				<RowDefinition Height="auto"/>
+				<RowDefinition Height="auto"/>
+				<RowDefinition Height="auto"/>
+				<RowDefinition Height="*"/>
+			</Grid.RowDefinitions>
+			<TextBlock Grid.Column="0" Grid.Row="0" Text="Source:"/>
+			<TextBlock Grid.Column="0" Grid.Row="1" Text="Destination:"/>
+			<TextBlock Grid.Column="0" Grid.Row="2" Text="Proxy:"/>
+
+			<TextBlock Grid.Column="2" Grid.Row="0">
+                <Hyperlink Command="{Binding OnSourcePathClicked}">
+                    <TextBlock Text="{Binding Build.FullPath, Mode=OneWay}"/>
+                </Hyperlink>
+            </TextBlock>
+			
+			<TextBlock Grid.Column="2" Grid.Row="1">
+                <Hyperlink Command="{Binding OnDestPathClicked}">
+                    <TextBlock Text="{Binding DstPath, Mode=OneWay}"/>
+                </Hyperlink>
+            </TextBlock>
+			
+			<TextBlock Grid.Column="2" Grid.Row="2" Text="{Binding ProxyStr}" />
+
+			<GroupBox Grid.Row="3" Grid.ColumnSpan="3" Header="Log" Visibility="{Binding LogLines.Count, Converter={StaticResource VisConv_Collapsed}}">
+				<GroupBox.ContextMenu>
+					<ContextMenu>
+						<MenuItem Header="Copy to Clipboard" Command="{Binding CopyLogClicked}"/>
+						<MenuItem Header="Save As..." Command="{Binding SaveLogClicked}"/>
+					</ContextMenu>
+				</GroupBox.ContextMenu>
+				<ScrollViewer VerticalScrollBarVisibility="Auto" HorizontalScrollBarVisibility="Auto"
+							  VerticalAlignment="Top" 
+							  MaxHeight="{Binding MaxLogScrollViewHeight}">
+					<ItemsControl ItemsSource="{Binding LogLines}" VirtualizingPanel.IsVirtualizing="True" Margin="4">
+						<ItemsControl.ItemTemplate>
+							<DataTemplate>
+								<TextBlock Text="{Binding Text}" FontFamily="Consolas" Foreground="{Binding Color, Converter={StaticResource ConsoleColorConverter}}"/>
+							</DataTemplate>
+						</ItemsControl.ItemTemplate>
+					</ItemsControl>
+				</ScrollViewer>
+			</GroupBox>
+		</Grid>
+	</Grid>
+</UserControl>
diff --git a/Engine/Source/Programs/UnsyncUI/UnsyncUI/JobControl.xaml.cs b/Engine/Source/Programs/UnsyncUI/UnsyncUI/JobControl.xaml.cs
new file mode 100644
--- /dev/null
+++ b/Engine/Source/Programs/UnsyncUI/UnsyncUI/JobControl.xaml.cs
@@ -0,0 +1,459 @@
+// Copyright Epic Games, Inc. All Rights Reserved.
+
+using System;
+using System.Collections.Generic;
+using System.Collections.ObjectModel;
+using System.IO;
+using System.Text;
+using System.Text.RegularExpressions;
+using System.Threading;
+using System.Threading.Tasks;
+using System.Windows;
+using System.Windows.Controls;
+
+namespace UnsyncUI
+{
+	/// <summary>
+	/// Interaction logic for ActiveJobControl.xaml
+	/// </summary>
+	public partial class JobControl : UserControl
+    {
+        public JobControl()
+        {
+            InitializeComponent();
+        }
+    }
+
+    public sealed class LogModel : BaseModel
+    {
+		private string text = null;
+        public string Text
+		{
+			get => text;
+			set => SetProperty(ref text, value);
+		}
+
+		private ConsoleColor color = ConsoleColor.Black;
+        public ConsoleColor Color
+		{
+			get => color;
+			set => SetProperty(ref color, value);
+		}
+    }
+
+	public sealed class FileModel : BaseModel
+	{
+		private string name = null;
+		public string Name
+		{
+			get => name;
+			set => SetProperty(ref name, value);
+		}
+
+		private string status = null;
+		public string Status
+		{
+			get => status;
+			set => SetProperty(ref status, value);
+		}
+	}
+
+    public sealed class JobModel : BaseModel
+    {
+        public BuildPlatformModel Build { get; }
+        public string DstPath { get; }
+        public bool DryRun { get; }
+        public string Proxy { get; }
+		public string DFS { get; }
+		public string AdditionalArgs { get; }
+		public string[] Exclusions { get; }
+
+        public string ProxyStr => string.IsNullOrWhiteSpace(Proxy) ? "(none)" : Proxy;
+
+        public string Name => $"{Build.Name}{(DryRun ? " - Dry Run" : "")}";
+
+        private double totalProgress = 0;
+        public double TotalProgress
+        {
+            get => totalProgress;
+            private set => SetProperty(ref totalProgress, value);
+        }
+
+        private bool isIndeterminate = true;
+        public bool IsIndeterminate
+        {
+            get => isIndeterminate;
+            private set => SetProperty(ref isIndeterminate, value);
+        }
+
+        private bool showProgressBar = false;
+        public bool ShowProgressBar
+        {
+            get => showProgressBar;
+            private set => SetProperty(ref showProgressBar, value);
+        }
+
+        private bool showCancelHyperlink = true;
+        public bool ShowCancelHyperlink
+        {
+            get => showCancelHyperlink;
+            private set => SetProperty(ref showCancelHyperlink, value);
+        }
+
+        private bool showClearHyperlink = false;
+        public bool ShowClearHyperlink
+        {
+            get => showClearHyperlink;
+            private set => SetProperty(ref showClearHyperlink, value);
+        }
+
+        private string statusMessage = null;
+        public string StatusMessage
+        {
+            get => statusMessage;
+            private set => SetProperty(ref statusMessage, value);
+        }
+
+		private ConsoleColor statusColor = ConsoleColor.Black;
+		public ConsoleColor StatusColor
+		{
+			get => statusColor;
+			private set => SetProperty(ref statusColor, value);
+		}
+
+		private string progressMessage = null;
+		public string ProgressMessage
+		{
+			get => progressMessage;
+			private set => SetProperty(ref progressMessage, value);
+		}
+
+
+		private bool detailsExpanded = false;
+        public bool DetailsExpanded
+        {
+            get => detailsExpanded;
+            set => SetProperty(ref detailsExpanded, value);
+        }
+
+        public Command OnCancelClicked { get; }
+        public Command OnClearClicked { get; }
+        public Command OnSourcePathClicked { get; }
+        public Command OnDestPathClicked { get; }
+		public Command CopyLogClicked { get; }
+		public Command SaveLogClicked { get; }
+
+		private TaskCompletionSource startTcs = new TaskCompletionSource();
+        private CancellationTokenSource cts = new CancellationTokenSource();
+        public bool IsCancelled => cts.IsCancellationRequested;
+		public bool Succeeded { get; private set; }
+
+		private double maxLogScrollViewHeight = Double.PositiveInfinity;
+		public double MaxLogScrollViewHeight
+		{
+			get => maxLogScrollViewHeight;
+			set => SetProperty(ref maxLogScrollViewHeight, value);
+		}
+
+		private Action<JobModel> onCompletion;
+        private Action<JobModel> onProgress;
+
+        public ObservableCollection<LogModel> LogLines { get; } = new ObservableCollection<LogModel>();
+		public ObservableCollection<FileModel> Files { get; } = new ObservableCollection<FileModel>();
+		private Dictionary<string, FileModel> filesMap = new Dictionary<string, FileModel>();
+
+        public JobModel(BuildPlatformModel build, string dstPath, bool dryRun, string proxy, string dfs, string additionalArgs, string[] exclusions, Action<JobModel> onCompletion, Action<JobModel> onClear, Action<JobModel> onProgress)
+        {
+            Build = build;
+            DstPath = dstPath;
+            DryRun = dryRun;
+            Proxy = proxy;
+			DFS = dfs;
+			AdditionalArgs = additionalArgs;
+			Exclusions = exclusions;
+			Succeeded = false;
+
+			OnCancelClicked = new Command(CancelJob) { Enabled = true };
+            OnClearClicked = new Command(() => onClear(this)) { Enabled = true };
+            OnSourcePathClicked = new Command(() => Shell.LaunchExplorer(Build.FullPath)) { Enabled = true };
+            OnDestPathClicked = new Command(() => Shell.LaunchExplorer(DstPath)) { Enabled = true };
+			CopyLogClicked = new Command(CopyLog) { Enabled = true };
+			SaveLogClicked = new Command(SaveLog) { Enabled = true };
+
+			this.onCompletion = onCompletion;
+			this.onProgress = onProgress;
+
+            cts.Token.Register(() => startTcs.TrySetCanceled());
+
+            RunAsync();
+        }
+
+		private string GetLogText()
+		{
+			var builder = new StringBuilder();
+			foreach (var logModel in LogLines)
+			{
+				builder.AppendLine(logModel.Text);
+			}
+			return builder.ToString();
+		}
+
+		private void CopyLog()
+		{
+			var log = GetLogText();
+			if (!string.IsNullOrWhiteSpace(log))
+			{
+				Clipboard.SetText(log, TextDataFormat.Text);
+			}
+		}
+
+		private void SaveLog()
+		{
+			var log = GetLogText();
+			if (!string.IsNullOrWhiteSpace(log))
+			{
+				var defaultName = $"unsync_{Build.SafeName}-{DateTime.Now.ToString("yyyy-MM-dd_HH-mm-ss")}.log";
+				var file = Shell.SelectSaveFile("Save Unsync Log", "Log Files|*.log", ".log", defaultName);
+				if (file != null)
+				{
+					try
+					{
+						File.WriteAllText(file, log);
+					}
+					catch (Exception ex)
+					{
+						MessageBox.Show($"Failed to save log to file: {ex}", "Save Log",
+							MessageBoxButton.OK, MessageBoxImage.Error);
+					}
+				}
+			}
+		}
+
+        public void StartJob() => startTcs.TrySetResult();
+
+		public void CancelJob()
+		{
+			StatusMessage = "Cancelling...";
+			OnCancelClicked.Enabled = false;
+			IsIndeterminate = true;
+			ProgressMessage = null;
+			StatusColor = ConsoleColor.Red;
+			cts.Cancel();
+		}
+
+        private async void RunAsync()
+        {
+			AsyncProcess proc = null;
+
+			try
+			{
+				// Wait for the job to be started.
+				await startTcs.Task;
+
+				DetailsExpanded = true;
+				IsIndeterminate = true;
+				ShowProgressBar = true;
+
+				StatusMessage = "Starting...";
+
+				LogModel currentLine = null;
+				void EmitLog(string line)
+				{
+					if (currentLine == null)
+					{
+						LogLines.Add(currentLine = new LogModel());
+					}
+
+					currentLine.Text = line;
+
+					// Determine log color
+					var color = ConsoleColor.Black;
+					if (currentLine.Text.Contains("ERROR:")) { color = ConsoleColor.Red; }
+					else if (currentLine.Text.Contains("WARNING:")) { color = ConsoleColor.DarkYellow; }
+
+					currentLine.Color = color;
+				}
+
+				void ProcessCommand(string message)
+				{
+					// Don't process command messages if we're cancelling, otherwise they'll overwrite UI state
+					if (cts.Token.IsCancellationRequested)
+						return;
+
+					var progressPattern = new Regex(@"@progress \[([^\]]*)\] ([0-9]+) / ([0-9]+)");
+					var statusPattern = new Regex(@"@status \[([^\]]*)\] (.*)");
+
+					Match match;
+					if ((match = progressPattern.Match(message)).Success)
+					{
+						var target = match.Groups[1].Value;
+						var numerator = ulong.Parse(match.Groups[2].Value);
+						var denominator = ulong.Parse(match.Groups[3].Value);
+
+						if (target == "*")
+						{
+							// Overall progress
+							TotalProgress = denominator > 0.0 ? Math.Clamp((numerator / (double)denominator) * 100.0, 0.0, 100.0) : 0.0;
+							ProgressMessage = $"{TotalProgress:0.00} %";
+							IsIndeterminate = false;
+
+							onProgress(this);
+						}
+						else
+						{
+							// @todo: specific file progress
+						}
+					}
+					else if ((match = statusPattern.Match(message)).Success)
+					{
+						var target = match.Groups[1].Value;
+						var text = match.Groups[2].Value;
+
+						if (target == "*")
+						{
+							// Overall status
+							StatusMessage = text;
+						}
+						else
+						{
+							if (!filesMap.TryGetValue(target, out var file))
+							{
+								filesMap.Add(target, file = new FileModel() { Name = target });
+								Files.Add(file);
+							}
+
+							file.Status = text;
+						}
+					}
+				}
+
+				var commandBuilder = new StringBuilder();
+				var lineBuilder = new StringBuilder();
+
+				char previous = (char)0;
+				bool command = false;
+
+				// Launch unsync.exe
+				var args = new List<string>() { "-v", "--progress" };
+				if (DryRun)
+				{
+					args.Add("-d");
+				}
+
+				if (!string.IsNullOrWhiteSpace(Proxy))
+				{
+					args.Add($"--proxy \"{Proxy}\"");
+				}
+
+				if (!string.IsNullOrWhiteSpace(DFS))
+				{
+					args.Add($"--dfs {DFS}");
+				}
+
+				if (Exclusions != null && Exclusions.Length > 0)
+				{
+					args.Add($"--exclude {string.Join(",", Exclusions)}");
+				}
+
+				args.Add($"\"{Build.FullPath}\"");
+				args.Add($"\"{DstPath}\"");
+
+				var argsStr = $"sync{(AdditionalArgs != null ? $" {AdditionalArgs}" : "")} {string.Join(" ", args)}";
+
+				// Allow completely overriding command line arguments (useful for testing)
+				if (AdditionalArgs != null && AdditionalArgs.StartsWith("!"))
+				{
+					argsStr = AdditionalArgs.Substring(1);
+				}
+
+				var unsyncPath = App.Current.UnsyncPath;
+				proc = new AsyncProcess(unsyncPath, argsStr);
+
+				await foreach (var str in proc.RunAsync(cts.Token))
+				{
+					foreach (var current in str)
+					{
+						if (current == '@' && (previous == '\r' || previous == '\n' || previous == (char)0))
+						{
+							command = true;
+						}
+
+						var builder = command 
+							? commandBuilder 
+							: lineBuilder;
+
+						switch (current)
+						{
+							case '\n':
+								if (command)
+								{
+									ProcessCommand(builder.ToString());
+									command = false;
+								}
+								else
+								{
+									if (previous != '\r')
+									{
+										EmitLog(builder.ToString());
+									}
+
+									currentLine = null;
+								}
+								builder.Clear();
+								break;
+
+							case '\r':
+								if (!command)
+								{
+									EmitLog(builder.ToString());
+									builder.Clear();
+								}
+								break;
+
+							default:
+								builder.Append(current);
+								break;
+						}
+
+						previous = current;
+					}
+				}
+			}
+			catch (OperationCanceledException)
+			{
+			}
+			finally
+			{
+				ShowProgressBar = false;
+				ShowCancelHyperlink = false;
+				ShowClearHyperlink = true;
+
+				if (proc?.ExitCode == 0)
+				{
+					StatusMessage = "Completed";
+					StatusColor = ConsoleColor.DarkGreen;
+					Succeeded = true;
+				}
+				else if (proc?.ExitCode != null)
+				{
+					StatusMessage = $"Failed (error code {proc?.ExitCode})";
+					StatusColor = ConsoleColor.Red;
+				}
+				else
+				{
+					StatusMessage = $"Failed (unknown error)";
+					StatusColor = ConsoleColor.Red;
+				}
+
+				ProgressMessage = null;
+				DetailsExpanded = false;
+
+				onCompletion(this);
+
+				// Allow the log to fill the entire window while job is running, but set a concrete limit on completed jobs.
+				// This is a workaround for missing scroll bar in the completed job logs.
+				MaxLogScrollViewHeight = 300;
+			}
+        }
+    }
+}
diff --git a/Engine/Source/Programs/UnsyncUI/UnsyncUI/MainWindow.xaml b/Engine/Source/Programs/UnsyncUI/UnsyncUI/MainWindow.xaml
new file mode 100644
--- /dev/null
+++ b/Engine/Source/Programs/UnsyncUI/UnsyncUI/MainWindow.xaml
@@ -0,0 +1,211 @@
+﻿<!--Copyright Epic Games, Inc. All Rights Reserved.-->
+
+<Window x:Class="UnsyncUI.MainWindow"
+        xmlns="http://schemas.microsoft.com/winfx/2006/xaml/presentation"
+        xmlns:x="http://schemas.microsoft.com/winfx/2006/xaml"
+        xmlns:d="http://schemas.microsoft.com/expression/blend/2008"
+        xmlns:mc="http://schemas.openxmlformats.org/markup-compatibility/2006"
+        xmlns:local="clr-namespace:UnsyncUI"
+        mc:Ignorable="d"
+        Title="UnsyncUI v1.0.7"
+        Height="800"
+        Width="1600"
+        MinWidth="600"
+        MinHeight="400" 
+		Closing="Window_Closing">
+	<Window.TaskbarItemInfo>
+		<TaskbarItemInfo ProgressState="{Binding ProgressState}" ProgressValue="{Binding ProgressValue}"/>
+	</Window.TaskbarItemInfo>
+	<Window.Resources>
+
+		<local:MainTabTemplateSelector x:Key="mainTabSelector">
+			<local:MainTabTemplateSelector.ProjectTemplate>
+				<DataTemplate>
+					<Grid>
+						<Grid.RowDefinitions>
+							<RowDefinition Height="auto" />
+							<RowDefinition Height="*" />
+						</Grid.RowDefinitions>
+
+						<local:BuildSelectorControl Grid.Row="1" />
+
+					</Grid>
+				</DataTemplate>
+			</local:MainTabTemplateSelector.ProjectTemplate>
+
+			<!-- Contents of the "Custom" project tab -->
+			<local:MainTabTemplateSelector.CustomTemplate>
+				<DataTemplate>
+					<Grid>
+						<Grid.RowDefinitions>
+							<RowDefinition Height="auto"/>
+							<RowDefinition Height="auto"/>
+							<RowDefinition Height="auto"/>
+							<RowDefinition Height="*"/>
+						</Grid.RowDefinitions>
+
+						<GroupBox Grid.Row="0" Header="Source">
+							<local:SelectFolderControl SelectedPath="{Binding SrcPath}" Description="Select a source directory to sync builds from."/>
+						</GroupBox>
+
+						<GroupBox Grid.Row="1" Header="Destination">
+							<local:SelectFolderControl SelectedPath="{Binding DstPath}" Description="Select a destination directory to sync builds to."/>
+						</GroupBox>
+
+						<TextBlock Grid.Row="2" Margin="4" HorizontalAlignment="Center">
+							<Hyperlink Command="{Binding OnSyncClicked}">Start Custom Sync</Hyperlink>
+						</TextBlock>
+
+					</Grid>
+				</DataTemplate>
+			</local:MainTabTemplateSelector.CustomTemplate>
+
+		</local:MainTabTemplateSelector>
+
+	</Window.Resources>
+	<Grid Margin="4">
+		<Grid.RowDefinitions>
+			<RowDefinition Height="auto" />
+			<RowDefinition Height="auto" />
+			<RowDefinition Height="*" />
+		</Grid.RowDefinitions>
+		<Grid.ColumnDefinitions>
+			<ColumnDefinition Width="476"/>
+			<ColumnDefinition Width="8"/>
+			<ColumnDefinition Width="*"/>
+		</Grid.ColumnDefinitions>
+
+		<GroupBox Grid.Row="0" Header="Proxy">
+			<ComboBox  ItemsSource="{Binding Config.Proxies}" SelectedItem="{Binding SelectedProxy}">
+				<ComboBox.ItemTemplate>
+					<DataTemplate>
+						<TextBlock>
+								<Run Text="{Binding Name}" />
+								<TextBlock Visibility="{Binding Path, Converter={StaticResource VisConv_Collapsed}}">
+									<Run Text=" ("/>
+									<Run Text="{Binding Path}"/>
+									<Run Text=")"/>
+								</TextBlock>
+						</TextBlock>
+					</DataTemplate>
+				</ComboBox.ItemTemplate>
+			</ComboBox>
+		</GroupBox>
+
+		<GroupBox Grid.Row="1" Header="Advanced">
+			<Grid>
+				<Grid.ColumnDefinitions>
+					<ColumnDefinition Width="auto"/>
+					<ColumnDefinition Width="*"/>
+					<ColumnDefinition Width="auto"/>
+				</Grid.ColumnDefinitions>
+
+				<Grid.RowDefinitions>
+					<RowDefinition/>
+					<RowDefinition/>
+				</Grid.RowDefinitions>
+
+				<CheckBox Grid.Column="0" VerticalAlignment="Center" IsChecked="{Binding DryRun}" Content="Dry Run" HorizontalAlignment="Center" Margin="32,4,32,4"/>
+				<TextBox Grid.Column="1" Text="{Binding DFS}"/>
+				<TextBlock Grid.Column="2" VerticalAlignment="Center" Text="DFS Filter" Margin="16,4,16,4"/>
+
+				<TextBox   Grid.Row="1" Grid.Column="0" Grid.ColumnSpan="2" Text="{Binding AdditionalArgs}" />
+				<TextBlock Grid.Row="1" Grid.Column="2" VerticalAlignment="Center" Text="Additional Args" Margin="16,4,16,4"/>
+			</Grid>
+		</GroupBox>
+
+		<GroupBox Grid.Row="2" Header="Projects">
+			<Grid>
+				<Grid.RowDefinitions>
+					<RowDefinition Height="auto"/>
+					<RowDefinition Height="*"/>
+				</Grid.RowDefinitions>
+				<TabControl Grid.Row ="1" ItemsSource="{Binding Tabs}" SelectedIndex="0" ContentTemplateSelector="{Binding Source={StaticResource mainTabSelector}}" SelectedItem="{Binding SelectedTab}">
+					<TabControl.ItemTemplate>
+						<DataTemplate>
+							<TextBlock Text="{Binding Name}"/>
+						</DataTemplate>
+					</TabControl.ItemTemplate>
+				</TabControl>
+			</Grid>
+		</GroupBox>
+
+		<GridSplitter Grid.Column="1" Grid.RowSpan="3" HorizontalAlignment="Stretch" Foreground="Transparent" Background="Transparent" />
+
+		<TextBlock Grid.Row="0" Grid.Column="2" Grid.RowSpan="3" TextAlignment="Center" VerticalAlignment="Center">
+			<TextBlock.Visibility>
+				<MultiBinding Converter="{StaticResource VisConv_All_Collapsed}">
+					<Binding Path="QueuedJobs.Count" Converter="{StaticResource NotConv}" />
+					<Binding Path="ActiveJob" Converter="{StaticResource NotConv}" />
+					<Binding Path="CompletedJobs.Count" Converter="{StaticResource NotConv}" />
+				</MultiBinding>
+			</TextBlock.Visibility>
+			<Run Text="To get started, search for builds in the left pane,"/>
+			<LineBreak/>
+			<Run Text="select one or more platforms and a destination folder, and click Start Sync."/>
+		</TextBlock>
+
+		<Grid Grid.Row="0" Grid.Column="2" Grid.RowSpan="3">
+			<Grid.Resources>
+				<GridLength x:Key="glStar">*</GridLength>
+				<GridLength x:Key="glAuto">auto</GridLength>
+				<local:ObjectConverter x:Key="rdHeightConv" TrueValue="{StaticResource glStar}" FalseValue="{StaticResource glAuto}"/>
+				<local:ObjectConverter x:Key="rdMinHeightConv" TrueValue="60" FalseValue="0"/>
+			</Grid.Resources>
+			<Grid.RowDefinitions>
+				<RowDefinition Height="auto"/>
+				<RowDefinition Height="{Binding ActiveJob, Converter={StaticResource rdHeightConv}}" MinHeight="{Binding ActiveJob, Converter={StaticResource rdMinHeightConv}}"/>
+				<RowDefinition Height="auto"/>
+			</Grid.RowDefinitions>
+
+			<!-- Queued Jobs Pane -->
+			<GroupBox Grid.Row="0" Visibility="{Binding QueuedJobs.Count, Converter={StaticResource VisConv_Collapsed}}">
+				<GroupBox.Header>
+					<TextBlock VerticalAlignment="Center" HorizontalAlignment="Center">
+                            <Run Text="{Binding QueuedJobs.Count, Mode=OneWay}" />
+                            <Run Text="Queued Job(s)" />
+                            <TextBlock>
+                                <Run Text=" - " />
+                                <Hyperlink Command="{Binding OnClearQueueClicked}">Cancel All</Hyperlink>
+                            </TextBlock>
+                        </TextBlock>
+				</GroupBox.Header>
+				<ItemsControl ItemsSource="{Binding QueuedJobs}">
+					<ItemsControl.ItemTemplate>
+						<ItemContainerTemplate>
+							<local:JobControl/>
+						</ItemContainerTemplate>
+					</ItemsControl.ItemTemplate>
+				</ItemsControl>
+			</GroupBox>
+
+			<!-- Active Job Pane -->
+			<GroupBox Grid.Row="1" Visibility="{Binding ActiveJob, Converter={StaticResource VisConv_Collapsed}}" Header="Active Job">
+				<local:JobControl DataContext="{Binding ActiveJob}"/>
+			</GroupBox>
+
+			<!-- Completed Jobs Pane -->
+			<GroupBox Grid.Row="2"  Visibility="{Binding CompletedJobs.Count, Converter={StaticResource VisConv_Collapsed}}">
+				<GroupBox.Header>
+					<TextBlock VerticalAlignment="Center" HorizontalAlignment="Center">
+                        <Run Text="{Binding CompletedJobs.Count, Mode=OneWay}" />
+                        <Run Text="Completed Job(s)" />
+                        <TextBlock Visibility="{Binding CompletedJobs.Count, Converter={StaticResource VisConv_Collapsed}}">
+                            <Run Text=" - " />
+                            <Hyperlink Command="{Binding OnClearCompletedClicked}">Clear All</Hyperlink>
+                        </TextBlock>
+                    </TextBlock>
+				</GroupBox.Header>
+
+				<ItemsControl ItemsSource="{Binding CompletedJobs}">
+					<ItemsControl.ItemTemplate>
+						<ItemContainerTemplate>
+							<local:JobControl/>
+						</ItemContainerTemplate>
+					</ItemsControl.ItemTemplate>
+				</ItemsControl>
+			</GroupBox>
+		</Grid>
+	</Grid>
+
+</Window>
diff --git a/Engine/Source/Programs/UnsyncUI/UnsyncUI/MainWindow.xaml.cs b/Engine/Source/Programs/UnsyncUI/UnsyncUI/MainWindow.xaml.cs
new file mode 100644
--- /dev/null
+++ b/Engine/Source/Programs/UnsyncUI/UnsyncUI/MainWindow.xaml.cs
@@ -0,0 +1,357 @@
+// Copyright Epic Games, Inc. All Rights Reserved.
+
+using System;
+using System.Collections.Generic;
+using System.Collections.ObjectModel;
+using System.ComponentModel;
+using System.Diagnostics;
+using System.Linq;
+using System.Windows;
+using System.Windows.Controls;
+using System.Windows.Shell;
+
+namespace UnsyncUI
+{
+	/// <summary>
+	/// Interaction logic for MainWindow.xaml
+	/// </summary>
+	public partial class MainWindow : Window
+	{
+		public MainWindow()
+		{
+			InitializeComponent();
+			DataContext = new MainWindowModel();
+		}
+
+		private void Window_Closing(object sender, CancelEventArgs e)
+		{
+			var model = DataContext as MainWindowModel;
+			if (model != null)
+			{
+				if (model.ActiveJob != null || model.QueuedJobs.Count > 0)
+				{
+					if (MessageBox.Show(
+						"There are active transfers in progress. Are you sure you want to quit? Transfers will be aborted.", 
+						"Quit Unsync",
+						MessageBoxButton.YesNo,
+						MessageBoxImage.Warning) != MessageBoxResult.Yes)
+					{
+						e.Cancel = true;
+					}
+				}
+			}
+		}
+	}
+
+	public sealed class MainTabTemplateSelector : DataTemplateSelector
+	{
+		public DataTemplate ProjectTemplate { get; set; }
+		public DataTemplate CustomTemplate { get; set; }
+
+		public override DataTemplate SelectTemplate(object item, DependencyObject container)
+			=> item is ProjectModel
+				? ProjectTemplate
+				: CustomTemplate;
+	}
+
+	public abstract class TabModel : BaseModel
+	{
+		public virtual void OnSelected() { }
+	}
+
+	public sealed class CustomModel : TabModel
+	{
+		private Action<IEnumerable<(string DstPath, string[] Exclusions, BuildPlatformModel Model)>> onBuildsSelected;
+
+		public string Name => "Custom";
+
+		private string srcPath = App.Current.UserConfig.CustomSrcPath;
+		public string SrcPath
+		{
+			get => srcPath;
+			set
+			{
+				SetProperty(ref srcPath, value);
+				App.Current.UserConfig.CustomSrcPath = value;
+
+				UpdateSyncCommand();
+			}
+		}
+
+		private string dstPath = App.Current.UserConfig.CustomDstPath;
+		public string DstPath
+		{
+			get => dstPath;
+			set
+			{
+				SetProperty(ref dstPath, value);
+				App.Current.UserConfig.CustomDstPath = value;
+
+				UpdateSyncCommand();
+			}
+		}
+
+		public Command OnSyncClicked { get; }
+
+		private void UpdateSyncCommand()
+			=> OnSyncClicked.Enabled = !string.IsNullOrWhiteSpace(SrcPath) && !string.IsNullOrWhiteSpace(DstPath);
+
+		public CustomModel(Action<IEnumerable<(string DstPath, string[] Exclusions, BuildPlatformModel Model)>> onBuildsSelected)
+		{
+			this.onBuildsSelected = onBuildsSelected;
+			OnSyncClicked = new Command(() =>
+			{
+				onBuildsSelected(new[]
+				{
+					(DstPath, default(string[]), new BuildPlatformModel(null, null, SrcPath))
+				});
+			});
+
+			UpdateSyncCommand();
+		}
+	}
+
+	public sealed class MainWindowModel : BaseModel
+	{
+		private Config.Proxy selectedProxy = 
+			App.Current.Config?.Proxies?.Where(p => p.Name.Equals(App.Current.UserConfig.Proxy, StringComparison.InvariantCultureIgnoreCase))?.FirstOrDefault() ??
+			App.Current.Config?.Proxies?.Where(p => p.Path != null).FirstOrDefault();
+
+		public Config.Proxy SelectedProxy
+		{
+			get => selectedProxy;
+			set
+			{
+				SetProperty(ref selectedProxy, value);
+				App.Current.UserConfig.Proxy = value.Name;
+			}
+		}
+
+		private bool dryRun = false;
+		public bool DryRun
+		{
+			get => dryRun;
+			set => SetProperty(ref dryRun, value);
+		}
+
+		private string dfs = App.Current.UserConfig.DFS ?? App.Current.Config?.DFS;
+		public string DFS
+		{
+			get => dfs;
+			set
+			{
+				SetProperty(ref dfs, value);
+				App.Current.UserConfig.DFS = value;
+			}
+		}
+
+		private bool showHelp = true;
+		public bool ShowHelp
+		{
+			get => showHelp;
+			set => SetProperty(ref showHelp, value);
+		}
+
+		private string additionalArgs = App.Current.UserConfig.AdditionalArgs;
+		public string AdditionalArgs
+		{
+			get => additionalArgs;
+			set
+			{
+				SetProperty(ref additionalArgs, value);
+				App.Current.UserConfig.AdditionalArgs = value;
+			}
+		}
+
+		private TaskbarItemProgressState progressState = TaskbarItemProgressState.None;
+		public TaskbarItemProgressState ProgressState
+		{
+			get => progressState;
+			set => SetProperty(ref progressState, value);
+		}
+
+		private float progressValue = 0.0f;
+		public float ProgressValue
+		{
+			get => progressValue;
+			set => SetProperty(ref progressValue, value);
+		}
+
+		private int batchJobCount = 0;
+
+		public Config Config => (Application.Current as App).Config;
+
+		public Command OnClearQueueClicked { get; }
+		public Command OnClearCompletedClicked { get; }
+
+		public ObservableCollection<TabModel> Tabs { get; } = new ObservableCollection<TabModel>();
+		public ObservableCollection<JobModel> QueuedJobs { get; } = new ObservableCollection<JobModel>();
+
+		private JobModel activeJob = null;
+		public JobModel ActiveJob 
+		{
+			get => activeJob;
+			set => SetProperty(ref activeJob, value);
+		}
+
+		public ObservableCollection<JobModel> CompletedJobs { get; } = new ObservableCollection<JobModel>();
+
+		private TabModel selectedTab = null;
+		public TabModel SelectedTab
+		{
+			get => selectedTab;
+			set
+			{
+				SetProperty(ref selectedTab, value);
+				selectedTab?.OnSelected();
+			}
+		}
+
+		public MainWindowModel()
+		{
+			OnClearQueueClicked = new Command(ClearQueue) { Enabled = true };
+			OnClearCompletedClicked = new Command(ClearCompleted) { Enabled = true };
+
+			if (Config != null)
+			{
+				foreach (var p in Config.Projects)
+				{
+					Tabs.Add(new ProjectModel(p, OnBuildsSelected));
+				}
+			}
+
+			Tabs.Add(new CustomModel(OnBuildsSelected));
+
+			SelectedTab = Tabs.FirstOrDefault();
+		}
+
+		private void ClearQueue()
+		{
+			foreach (var job in QueuedJobs)
+				job.CancelJob();
+		}
+
+		private void ClearCompleted()
+		{
+			CompletedJobs.Clear();
+			UpdateProgressState();
+		}
+
+		public void OnBuildsSelected(IEnumerable<(string DstPath, string[] Exclusions, BuildPlatformModel Model)> selectedBuilds)
+		{
+			foreach (var build in selectedBuilds)
+			{
+				AddJob(new JobModel(build.Model, build.DstPath, DryRun, SelectedProxy?.Path, DFS, AdditionalArgs, build.Exclusions, OnJobCompleted, OnClearJob, OnJobProgress));
+			}
+		}
+
+		private void AddJob(JobModel job)
+		{
+			++batchJobCount;
+
+			if (ActiveJob == null)
+			{
+				ActiveJob = job;
+				job.StartJob();
+			}
+			else
+			{
+				QueuedJobs.Add(job);
+			}
+
+			UpdateProgressValue();
+			UpdateProgressState();
+		}
+
+		private void OnJobProgress(JobModel job)
+		{
+			UpdateProgressValue();
+		}
+
+		private void OnJobCompleted(JobModel job)
+		{
+			QueuedJobs.Remove(job);
+
+			if (!job.IsCancelled)
+			{
+				CompletedJobs.Add(job);
+			}
+			else
+			{
+				Debug.Assert(batchJobCount > 0);
+				--batchJobCount;
+			}
+
+			if (job == ActiveJob)
+			{
+				ActiveJob = null;
+				if (QueuedJobs.Count > 0)
+				{
+					var newJob = QueuedJobs.Where(q => !q.IsCancelled).LastOrDefault();
+					if (newJob != null)
+					{
+						QueuedJobs.Remove(newJob);
+						ActiveJob = newJob;
+
+						newJob.StartJob();
+					}
+				}
+				else
+				{
+					batchJobCount = 0;
+				}
+			}
+
+			UpdateProgressValue();
+			UpdateProgressState();
+		}
+
+		private void OnClearJob(JobModel job)
+		{
+			CompletedJobs.Remove(job);
+			UpdateProgressState();
+		}
+
+		private void UpdateProgressState()
+		{
+			bool error = false;
+
+			foreach (var job in CompletedJobs)
+			{
+				if (!job.Succeeded)
+					error = true;
+			}
+
+			if(error)
+			{
+				ProgressState = TaskbarItemProgressState.Error;
+			}
+			else if(batchJobCount > 0)
+			{
+				ProgressState = TaskbarItemProgressState.Normal;
+			}
+			else
+			{
+				ProgressState = TaskbarItemProgressState.None;
+			}
+		}
+		private void UpdateProgressValue()
+		{
+			if(batchJobCount == 0)
+			{
+				ProgressValue = 1.0f;
+			}
+			else
+			{
+				float progress = batchJobCount - QueuedJobs.Count;
+
+				if(ActiveJob != null)
+				{
+					progress += (float)ActiveJob.TotalProgress / 100.0f - 1;
+				}
+
+				ProgressValue = progress / batchJobCount;
+			}
+		}
+	}
+}
diff --git a/Engine/Source/Programs/UnsyncUI/UnsyncUI/SelectFolderControl.xaml b/Engine/Source/Programs/UnsyncUI/UnsyncUI/SelectFolderControl.xaml
new file mode 100644
--- /dev/null
+++ b/Engine/Source/Programs/UnsyncUI/UnsyncUI/SelectFolderControl.xaml
@@ -0,0 +1,22 @@
+﻿<!--Copyright Epic Games, Inc. All Rights Reserved.-->
+
+<UserControl x:Class="UnsyncUI.SelectFolderControl"
+             xmlns="http://schemas.microsoft.com/winfx/2006/xaml/presentation"
+             xmlns:x="http://schemas.microsoft.com/winfx/2006/xaml"
+             xmlns:mc="http://schemas.openxmlformats.org/markup-compatibility/2006" 
+             xmlns:d="http://schemas.microsoft.com/expression/blend/2008" 
+             xmlns:local="clr-namespace:UnsyncUI"
+             mc:Ignorable="d"
+             d:DesignWidth="350">
+    <Grid>
+        <Grid.ColumnDefinitions>
+            <ColumnDefinition Width="*"/>
+            <ColumnDefinition Width="auto"/>
+        </Grid.ColumnDefinitions>
+        <TextBox 
+            Grid.Column="0"
+            Text="{Binding SelectedPath, Mode=TwoWay, UpdateSourceTrigger=PropertyChanged, RelativeSource={RelativeSource Mode=FindAncestor, AncestorType=UserControl}}" 
+            ToolTip="{Binding Description, RelativeSource={RelativeSource Mode=FindAncestor, AncestorType=UserControl}}"/>
+        <Button Grid.Column="1" Content=" ... " Click="OnBrowseClicked" ToolTip="Browse ..."/>
+    </Grid>
+</UserControl>
diff --git a/Engine/Source/Programs/UnsyncUI/UnsyncUI/SelectFolderControl.xaml.cs b/Engine/Source/Programs/UnsyncUI/UnsyncUI/SelectFolderControl.xaml.cs
new file mode 100644
--- /dev/null
+++ b/Engine/Source/Programs/UnsyncUI/UnsyncUI/SelectFolderControl.xaml.cs
@@ -0,0 +1,43 @@
+// Copyright Epic Games, Inc. All Rights Reserved.
+
+using System.Windows;
+using System.Windows.Controls;
+using System.Windows.Data;
+
+namespace UnsyncUI
+{
+	/// <summary>
+	/// Interaction logic for SelectFolderControl.xaml
+	/// </summary>
+	public partial class SelectFolderControl : UserControl
+    {
+        public static readonly DependencyProperty DescriptionProperty = DependencyProperty.Register(nameof(Description), typeof(string), typeof(SelectFolderControl));
+        public static readonly DependencyProperty SelectedPathProperty = DependencyProperty.Register(nameof(SelectedPath), typeof(string), typeof(SelectFolderControl), new FrameworkPropertyMetadata()
+        {
+            BindsTwoWayByDefault = true,
+            DefaultUpdateSourceTrigger = UpdateSourceTrigger.LostFocus
+        });
+
+        public string Description
+        {
+            get => (string)GetValue(DescriptionProperty);
+            set => SetValue(DescriptionProperty, value);
+        }
+
+        public string SelectedPath
+        {
+            get => (string)GetValue(SelectedPathProperty);
+            set => SetValue(SelectedPathProperty, value);
+        }
+
+        public SelectFolderControl()
+        {
+            InitializeComponent();
+        }
+
+        public void OnBrowseClicked(object sender, RoutedEventArgs e)
+        {
+            SelectedPath = Shell.SelectFolder(SelectedPath, Description);
+        }
+    }
+}
diff --git a/Engine/Source/Programs/UnsyncUI/UnsyncUI/Shell.cs b/Engine/Source/Programs/UnsyncUI/UnsyncUI/Shell.cs
new file mode 100644
--- /dev/null
+++ b/Engine/Source/Programs/UnsyncUI/UnsyncUI/Shell.cs
@@ -0,0 +1,54 @@
+// Copyright Epic Games, Inc. All Rights Reserved.
+
+using System;
+using System.Diagnostics;
+using System.Threading.Tasks;
+using System.Windows.Forms;
+
+namespace UnsyncUI
+{
+	public static class Shell
+	{
+		public static void LaunchExplorer(string path)
+		{
+			Task.Run(() => Process.Start("explorer.exe", path));
+		}
+
+		public static string SelectFolder(string root, string desc)
+		{
+			using FolderBrowserDialog fbd = new FolderBrowserDialog()
+			{
+				AutoUpgradeEnabled = true,
+				RootFolder = Environment.SpecialFolder.Desktop,
+				ShowNewFolderButton = true,
+				SelectedPath = root,
+				Description = string.IsNullOrWhiteSpace(desc) ? string.Empty : desc,
+				UseDescriptionForTitle = !string.IsNullOrWhiteSpace(desc)
+			};
+
+			if (fbd.ShowDialog() == DialogResult.OK)
+			{
+				return fbd.SelectedPath;
+			}
+
+			return root;
+		}
+
+		public static string SelectSaveFile(string title, string filter, string defaultExt, string defaultFilename)
+		{
+			using SaveFileDialog sfd = new SaveFileDialog()
+			{
+				AddExtension = true,
+				Filter = filter,
+				Title = title,
+				DefaultExt = defaultExt,
+				FileName = defaultFilename
+			};
+
+			if (sfd.ShowDialog() == DialogResult.OK)
+				return sfd.FileName;
+
+			return null;
+		}
+	}
+}
diff --git a/Engine/Source/Programs/UnsyncUI/UnsyncUI/UnsyncUI.csproj b/Engine/Source/Programs/UnsyncUI/UnsyncUI/UnsyncUI.csproj
new file mode 100644
--- /dev/null
+++ b/Engine/Source/Programs/UnsyncUI/UnsyncUI/UnsyncUI.csproj
@@ -0,0 +1,13 @@
+﻿<Project Sdk="Microsoft.NET.Sdk">
+
+	<PropertyGroup>
+		<OutputType>WinExe</OutputType>
+		<TargetFramework>net5.0-windows</TargetFramework>
+		<UseWPF>true</UseWPF>
+		<UseWindowsForms>true</UseWindowsForms>
+		<ApplicationManifest>app.manifest</ApplicationManifest>
+        <IncludeAllContentForSelfExtract>true</IncludeAllContentForSelfExtract>
+        <ApplicationIcon>unsync_48px.ico</ApplicationIcon>
+	</PropertyGroup>
+
+</Project>
diff --git a/Engine/Source/Programs/UnsyncUI/UnsyncUI/Utils.cs b/Engine/Source/Programs/UnsyncUI/UnsyncUI/Utils.cs
new file mode 100644
--- /dev/null
+++ b/Engine/Source/Programs/UnsyncUI/UnsyncUI/Utils.cs
@@ -0,0 +1,290 @@
+// Copyright Epic Games, Inc. All Rights Reserved.
+
+using System;
+using System.Collections.Generic;
+using System.ComponentModel;
+using System.Diagnostics;
+using System.IO;
+using System.Linq;
+using System.Runtime.CompilerServices;
+using System.Runtime.InteropServices;
+using System.Text;
+using System.Threading;
+using System.Threading.Tasks;
+using System.Threading.Tasks.Dataflow;
+using System.Windows.Input;
+
+namespace UnsyncUI
+{
+	public abstract class BaseModel : INotifyPropertyChanged
+    {
+        public event PropertyChangedEventHandler PropertyChanged;
+
+        protected void OnPropertyChanged(string name)
+            => PropertyChanged?.Invoke(this, new PropertyChangedEventArgs(name));
+
+        protected void SetProperty<T>(ref T property, T newValue, [CallerMemberName] string name = null)
+        {
+            if (!Equals(property, newValue))
+            {
+                property = newValue;
+                OnPropertyChanged(name);
+            }
+        }
+    }
+
+    public class Command<T> : BaseModel, ICommand
+    {
+        public event EventHandler CanExecuteChanged;
+        private Action<T> action;
+        private bool enabled = false;
+
+        public bool Enabled
+        {
+            get => enabled;
+            set
+            {
+                if (enabled != value)
+                {
+                    enabled = value;
+                    CanExecuteChanged?.Invoke(this, new EventArgs());
+					OnPropertyChanged(nameof(Enabled));
+                }
+            }
+        }
+
+        public Command(Action<T> action)
+        {
+            this.action = action;
+        }
+
+        public bool CanExecute(object parameter) => Enabled;
+        public void Execute(object parameter) => action((T)parameter);
+    }
+
+    public sealed class Command : Command<object>
+    {
+        public Command(Action action)
+            : base(o => action())
+        { }
+    }
+
+    public static class AsyncIO
+    {
+		private static SemaphoreSlim mutex = new SemaphoreSlim(4);
+
+        public static Task<IEnumerable<string>> EnumerateDirectoriesAsync(string path, CancellationToken token)
+        {
+			var tcs = new TaskCompletionSource<IEnumerable<string>>();
+            Task.Run(() =>
+			{
+				using var cancel = token.Register(() => tcs.TrySetCanceled());
+				try
+				{
+					mutex.Wait(token);
+					try
+					{
+						var timer = new Stopwatch();
+						timer.Start();
+
+						var dirs = Directory.EnumerateDirectories(path).ToList();
+
+						timer.Stop();
+						Debug.WriteLine($"Time: {timer.Elapsed.TotalSeconds:0.000} s - {path}");
+
+						tcs.TrySetResult(dirs);
+					}
+					catch (Exception ex)
+					{
+						tcs.TrySetException(ex);
+					}
+
+					mutex.Release();
+				}
+				catch (OperationCanceledException)
+				{
+					// Mutex wait was cancelled. Don't release
+					tcs.TrySetCanceled();
+				}
+			});
+
+            return tcs.Task;
+		}
+    }
+
+    public sealed class AsyncProcess
+    {
+		public static class ProcessTracker
+		{
+			[Flags]
+			public enum JOBOBJECTLIMIT : uint
+			{
+				JOB_OBJECT_LIMIT_KILL_ON_JOB_CLOSE = 0x2000
+			}
+
+			[StructLayout(LayoutKind.Sequential)]
+			public struct JOBOBJECT_BASIC_LIMIT_INFORMATION
+			{
+				public long PerProcessUserTimeLimit;
+				public long PerJobUserTimeLimit;
+				public JOBOBJECTLIMIT LimitFlags;
+				public UIntPtr MinimumWorkingSetSize;
+				public UIntPtr MaximumWorkingSetSize;
+				public uint ActiveProcessLimit;
+				public long Affinity;
+				public uint PriorityClass;
+				public uint SchedulingClass;
+			}
+
+			[StructLayout(LayoutKind.Sequential)]
+			public struct IO_COUNTERS
+			{
+				public ulong ReadOperationCount;
+				public ulong WriteOperationCount;
+				public ulong OtherOperationCount;
+				public ulong ReadTransferCount;
+				public ulong WriteTransferCount;
+				public ulong OtherTransferCount;
+			}
+
+			[StructLayout(LayoutKind.Sequential)]
+			public struct JOBOBJECT_EXTENDED_LIMIT_INFORMATION
+			{
+				public JOBOBJECT_BASIC_LIMIT_INFORMATION BasicLimitInformation;
+				public IO_COUNTERS IoInfo;
+				public UIntPtr ProcessMemoryLimit;
+				public UIntPtr JobMemoryLimit;
+				public UIntPtr PeakProcessMemoryUsed;
+				public UIntPtr PeakJobMemoryUsed;
+			}
+
+			public enum JobObjectInfoType
+			{
+				ExtendedLimitInformation = 9
+			}
+
+			[DllImport("kernel32.dll", CharSet = CharSet.Unicode)]
+			private extern static IntPtr CreateJobObject(IntPtr lpJobAttributes, string name);
+
+			[DllImport("kernel32.dll")]
+			private extern static bool SetInformationJobObject(IntPtr job, JobObjectInfoType infoType, IntPtr lpJobObjectInfo, uint cbJobObjectInfoLength);
+
+			[DllImport("kernel32.dll", SetLastError = true)]
+			private extern static bool AssignProcessToJobObject(IntPtr job, IntPtr process);
+
+			private static readonly IntPtr hJob = CreateJobObject(IntPtr.Zero, $"ProcessTracker_{Process.GetCurrentProcess().Id}");
+
+			static ProcessTracker()
+			{
+				if (hJob == IntPtr.Zero)
+					throw new Win32Exception(Marshal.GetLastWin32Error(), "Failed to create job handle.");
+
+				IntPtr infoExPtr = IntPtr.Zero;
+
+				try
+				{
+					var size = Marshal.SizeOf(typeof(JOBOBJECT_EXTENDED_LIMIT_INFORMATION));
+					infoExPtr = Marshal.AllocHGlobal(size);
+
+					Marshal.StructureToPtr(new JOBOBJECT_EXTENDED_LIMIT_INFORMATION()
+					{
+						BasicLimitInformation = new JOBOBJECT_BASIC_LIMIT_INFORMATION()
+						{
+							LimitFlags = JOBOBJECTLIMIT.JOB_OBJECT_LIMIT_KILL_ON_JOB_CLOSE
+						}
+					}, infoExPtr, false);
+
+					if (!SetInformationJobObject(hJob, JobObjectInfoType.ExtendedLimitInformation, infoExPtr, (uint)size))
+						throw new Win32Exception(Marshal.GetLastWin32Error(), "Failed to set job info.");
+				}
+				finally
+				{
+					if (infoExPtr != IntPtr.Zero)
+						Marshal.FreeHGlobal(infoExPtr);
+				}
+			}
+
+			public static void Attach(Process proc)
+			{
+				if (!AssignProcessToJobObject(hJob, proc.Handle) && !proc.HasExited)
+					throw new Win32Exception();
+			}
+		}
+
+		private Process proc = new Process();
+		public int? ExitCode { get; private set; }
+
+		public AsyncProcess(string path, string args)
+		{
+			proc.StartInfo = new ProcessStartInfo(path, args)
+			{
+				CreateNoWindow = true,
+				UseShellExecute = false,
+				WorkingDirectory = Environment.CurrentDirectory,
+				RedirectStandardOutput = true,
+				RedirectStandardError = true
+			};
+		}
+
+		public async IAsyncEnumerable<string> RunAsync([EnumeratorCancellation] CancellationToken cancelToken)
+		{
+			try
+			{
+				using (var cancel = cancelToken.Register(() =>
+				{
+					if (!proc.HasExited)
+					{
+						proc.Kill();
+					}
+				}))
+				{
+					proc.Start();
+					ProcessTracker.Attach(proc);
+
+					var pipe = new BufferBlock<string>();
+
+					async Task ReadStream(Stream stream)
+					{
+						var block = new byte[4096];
+						var mem = new Memory<byte>(block);
+
+						while (true)
+						{
+							int bytesRead = await stream.ReadAsync(mem);
+							if (bytesRead == 0)
+								break;
+
+							// @todo: this won't handle UTF-8 encoding if a char is split across the read boundary.
+							pipe.Post(Encoding.UTF8.GetString(mem.Span.Slice(0, bytesRead)));
+						}
+
+						pipe.Complete();
+					}
+
+					var stdoutTask = ReadStream(proc.StandardOutput.BaseStream);
+					var stderrTask = ReadStream(proc.StandardError.BaseStream);
+
+					var completionTask = Task.Run(async () =>
+					{
+						await Task.WhenAll(stdoutTask, stderrTask);
+						pipe.Complete();
+					});
+
+					while (await pipe.OutputAvailableAsync())
+					{
+						yield return pipe.Receive();
+					}
+
+					await completionTask;
+				}
+			}
+			finally
+			{
+				ExitCode = proc.ExitCode;
+				proc.Dispose();
+
+				cancelToken.ThrowIfCancellationRequested();
+			}
+		}
+	}
+}
diff --git a/Engine/Source/Programs/UnsyncUI/UnsyncUI/app.manifest b/Engine/Source/Programs/UnsyncUI/UnsyncUI/app.manifest
new file mode 100644
--- /dev/null
+++ b/Engine/Source/Programs/UnsyncUI/UnsyncUI/app.manifest
@@ -0,0 +1,76 @@
+<?xml version="1.0" encoding="utf-8"?>
+<assembly manifestVersion="1.0" xmlns="urn:schemas-microsoft-com:asm.v1">
+  <assemblyIdentity version="1.0.0.0" name="MyApplication.app"/>
+  <trustInfo xmlns="urn:schemas-microsoft-com:asm.v2">
+    <security>
+      <requestedPrivileges xmlns="urn:schemas-microsoft-com:asm.v3">
+        <!-- UAC Manifest Options
+             If you want to change the Windows User Account Control level replace the 
+             requestedExecutionLevel node with one of the following.
+
+        <requestedExecutionLevel  level="asInvoker" uiAccess="false" />
+        <requestedExecutionLevel  level="requireAdministrator" uiAccess="false" />
+        <requestedExecutionLevel  level="highestAvailable" uiAccess="false" />
+
+            Specifying requestedExecutionLevel element will disable file and registry virtualization. 
+            Remove this element if your application requires this virtualization for backwards
+            compatibility.
+        -->
+        <requestedExecutionLevel level="asInvoker" uiAccess="false" />
+      </requestedPrivileges>
+    </security>
+  </trustInfo>
+
+  <compatibility xmlns="urn:schemas-microsoft-com:compatibility.v1">
+    <application>
+      <!-- A list of the Windows versions that this application has been tested on
+           and is designed to work with. Uncomment the appropriate elements
+           and Windows will automatically select the most compatible environment. -->
+
+      <!-- Windows Vista -->
+      <!--<supportedOS Id="{e2011457-1546-43c5-a5fe-008deee3d3f0}" />-->
+
+      <!-- Windows 7 -->
+      <!--<supportedOS Id="{35138b9a-5d96-4fbd-8e2d-a2440225f93a}" />-->
+
+      <!-- Windows 8 -->
+      <!--<supportedOS Id="{4a2f28e3-53b9-4441-ba9c-d69d4a4a6e38}" />-->
+
+      <!-- Windows 8.1 -->
+      <!--<supportedOS Id="{1f676c76-80e1-4239-95bb-83d0f6d0da78}" />-->
+
+      <!-- Windows 10 -->
+      <!--<supportedOS Id="{8e0f7a12-bfb3-4fe8-b9a5-48fd50a15a9a}" />-->
+
+    </application>
+  </compatibility>
+
+  <!-- Indicates that the application is DPI-aware and will not be automatically scaled by Windows at higher
+       DPIs. Windows Presentation Foundation (WPF) applications are automatically DPI-aware and do not need 
+       to opt in. Windows Forms applications targeting .NET Framework 4.6 that opt into this setting, should 
+       also set the 'EnableWindowsFormsHighDpiAutoResizing' setting to 'true' in their app.config. -->
+  <!--
+  <application xmlns="urn:schemas-microsoft-com:asm.v3">
+    <windowsSettings>
+      <dpiAware xmlns="http://schemas.microsoft.com/SMI/2005/WindowsSettings">true</dpiAware>
+    </windowsSettings>
+  </application>
+  -->
+
+  <!-- Enable themes for Windows common controls and dialogs (Windows XP and later) -->
+
+  <dependency>
+    <dependentAssembly>
+      <assemblyIdentity
+          type="win32"
+          name="Microsoft.Windows.Common-Controls"
+          version="6.0.0.0"
+          processorArchitecture="*"
+          publicKeyToken="6595b64144ccf1df"
+          language="*"
+        />
+    </dependentAssembly>
+  </dependency>
+
+
+</assembly>
