UnsyncUI - Automatically discover the proxy server list from a seed server specified by host name in the config, bump version to 1.0.8

diff --git a/Engine/Source/Programs/UnsyncUI/UnsyncUI/Config.cs b/Engine/Source/Programs/UnsyncUI/UnsyncUI/Config.cs
--- a/Engine/Source/Programs/UnsyncUI/UnsyncUI/Config.cs
+++ b/Engine/Source/Programs/UnsyncUI/UnsyncUI/Config.cs
@@ -9,11 +9,23 @@ using System.Threading;
 using System.Threading.Tasks;
 using System.Threading.Tasks.Dataflow;
 using System.Xml.Linq;
+using System.Net;
+using System.Text.Json;
 
 namespace UnsyncUI
 {
 	public sealed class Config
 	{
+		// Structure that represents mirror server entry in the list from /api/v1/mirrors endpoint
+		class JsonMirrorDesc
+		{
+			public String name { get; set; }
+			public String address { get; set; }
+			public int port { get; set; } = 0;
+			public String description { get; set; }
+			public String parent { get; set; }
+		}
+
 		public sealed class Proxy
 		{
 			public string Name { get; set; }
@@ -159,12 +171,27 @@ namespace UnsyncUI
 				Name = "(none)",
 				Path = null
 			});
-			Proxies.AddRange(rootNode.Element("proxies").Elements("proxy").Select(p => new Proxy()
+
+			List<Proxy> ConfigProxies = new List<Proxy>();
+
+			ConfigProxies.AddRange(rootNode.Element("proxies").Elements("proxy").Select(p => new Proxy()
 			{
 				Name = p.Attribute("name")?.Value,
 				Path = p.Attribute("path")?.Value
 			}));
 
+			// Auto-discover proxies
+			List<Proxy> DiscoveredProxies = DiscoverProxies(ConfigProxies);
+
+			if (DiscoveredProxies == null)
+			{
+				Proxies.AddRange(ConfigProxies);
+			}
+			else
+			{
+				Proxies.AddRange(DiscoveredProxies);
+			}
+
 			Projects = rootNode.Element("projects").Elements("project").Select(p => new Project()
 			{
 				Name = p.Attribute("name")?.Value,
@@ -174,5 +201,83 @@ namespace UnsyncUI
 				Exclusions = p.Elements("exclude").Select(d => d.Value).ToList()
 			}).ToList();
 		}
+
+		private List<Proxy> DiscoverProxies(List<Proxy> SeedServers)
+		{
+			int DefaultPort = 53841;
+
+			foreach (Proxy SeedServer in SeedServers)
+			{
+				if (SeedServer.Path == null)
+				{
+					continue;
+				}
+
+				try
+				{
+					String Url = SeedServer.Path;
+					if (!Url.Contains(":"))
+					{
+						Url += ":" + DefaultPort.ToString();
+					}
+
+					if (!Url.StartsWith("http://"))
+					{
+						Url = "http://" + Url;
+					}
+
+					Url += "/api/v1/mirrors";
+
+					var Request = WebRequest.Create(Url);
+					Request.Timeout = 2500;
+					var Response = (HttpWebResponse)Request.GetResponse();
+					if (Response.StatusCode == HttpStatusCode.OK)
+					{
+						var Reader = new StreamReader(Response.GetResponseStream());
+						var Body = Reader.ReadToEnd();
+						var ParsedList = JsonSerializer.Deserialize<List<JsonMirrorDesc>>(Body);
+
+						var ParsedProxies = new List<Proxy>();
+						foreach (var ParsedProxy in ParsedList)
+						{
+							if (ParsedProxy.address == null)
+							{
+								continue;
+							}
+
+							var ConvertedProxy = new Proxy();
+							ConvertedProxy.Path = ParsedProxy.address;
+
+							if (ParsedProxy.description != null)
+							{
+								ConvertedProxy.Name = ParsedProxy.description;
+							}
+							else if (ParsedProxy.name != null)
+							{
+								ConvertedProxy.Name = ParsedProxy.name;
+							}
+
+							if (ParsedProxy.port != 0)
+							{
+								ConvertedProxy.Path += ":" + ParsedProxy.port.ToString();
+							}
+
+							ParsedProxies.Add(ConvertedProxy);
+						}
+
+						if (ParsedProxies.Count != 0)
+						{
+							return ParsedProxies;
+						}
+					}
+				}
+				catch (Exception)
+				{
+					continue;
+				}
+			}
+
+			return null;
+		}
 	}
 }
diff --git a/Engine/Source/Programs/UnsyncUI/UnsyncUI/MainWindow.xaml b/Engine/Source/Programs/UnsyncUI/UnsyncUI/MainWindow.xaml
--- a/Engine/Source/Programs/UnsyncUI/UnsyncUI/MainWindow.xaml
+++ b/Engine/Source/Programs/UnsyncUI/UnsyncUI/MainWindow.xaml
@@ -7,7 +7,7 @@
         xmlns:mc="http://schemas.openxmlformats.org/markup-compatibility/2006"
         xmlns:local="clr-namespace:UnsyncUI"
         mc:Ignorable="d"
-        Title="UnsyncUI v1.0.7"
+        Title="UnsyncUI v1.0.8"
         Height="800"
         Width="1600"
         MinWidth="600"
