UnsyncUI - Fix virtualization of the log control, fixes performance issue when processing a large number of log lines

diff --git a/Engine/Source/Programs/UnsyncUI/UnsyncUI/JobControl.xaml b/Engine/Source/Programs/UnsyncUI/UnsyncUI/JobControl.xaml
--- a/Engine/Source/Programs/UnsyncUI/UnsyncUI/JobControl.xaml
+++ b/Engine/Source/Programs/UnsyncUI/UnsyncUI/JobControl.xaml
@@ -83,17 +83,25 @@
 						<MenuItem Header="Save As..." Command="{Binding SaveLogClicked}"/>
 					</ContextMenu>
 				</GroupBox.ContextMenu>
-				<ScrollViewer VerticalScrollBarVisibility="Auto" HorizontalScrollBarVisibility="Auto"
-							  VerticalAlignment="Top" 
-							  MaxHeight="{Binding MaxLogScrollViewHeight}">
-					<ItemsControl ItemsSource="{Binding LogLines}" VirtualizingPanel.IsVirtualizing="True" Margin="4">
-						<ItemsControl.ItemTemplate>
-							<DataTemplate>
-								<TextBlock Text="{Binding Text}" FontFamily="Consolas" Foreground="{Binding Color, Converter={StaticResource ConsoleColorConverter}}"/>
-							</DataTemplate>
-						</ItemsControl.ItemTemplate>
-					</ItemsControl>
-				</ScrollViewer>
+				<ItemsControl ItemsSource="{Binding LogLines}">
+					<ItemsControl.ItemTemplate>
+						<DataTemplate>
+							<TextBlock Text="{Binding Text}" FontFamily="Consolas" Foreground="{Binding Color, Converter={StaticResource ConsoleColorConverter}}"/>
+						</DataTemplate>
+					</ItemsControl.ItemTemplate>
+					<ItemsControl.ItemsPanel>
+						<ItemsPanelTemplate>
+							<VirtualizingStackPanel IsVirtualizing="True" VirtualizationMode="Recycling" />
+						</ItemsPanelTemplate>
+					</ItemsControl.ItemsPanel>
+					<ItemsControl.Template>
+						<ControlTemplate TargetType="ItemsControl">
+							<ScrollViewer CanContentScroll="True" Focusable="True">
+								<ItemsPresenter />
+							</ScrollViewer>
+						</ControlTemplate>
+					</ItemsControl.Template>
+				</ItemsControl>
 			</GroupBox>
 		</Grid>
 	</Grid>
