UnsyncUI - Fix slow expansion of completed job logs, bump version to 1.0.9

diff --git a/Engine/Source/Programs/UnsyncUI/UnsyncUI/JobControl.xaml b/Engine/Source/Programs/UnsyncUI/UnsyncUI/JobControl.xaml
--- a/Engine/Source/Programs/UnsyncUI/UnsyncUI/JobControl.xaml
+++ b/Engine/Source/Programs/UnsyncUI/UnsyncUI/JobControl.xaml
@@ -96,7 +96,7 @@
 					</ItemsControl.ItemsPanel>
 					<ItemsControl.Template>
 						<ControlTemplate TargetType="ItemsControl">
-							<ScrollViewer CanContentScroll="True" Focusable="True">
+							<ScrollViewer CanContentScroll="True" Focusable="True" MaxHeight="{Binding MaxLogScrollViewHeight}">
 								<ItemsPresenter />
 							</ScrollViewer>
 						</ControlTemplate>
diff --git a/Engine/Source/Programs/UnsyncUI/UnsyncUI/MainWindow.xaml b/Engine/Source/Programs/UnsyncUI/UnsyncUI/MainWindow.xaml
--- a/Engine/Source/Programs/UnsyncUI/UnsyncUI/MainWindow.xaml
+++ b/Engine/Source/Programs/UnsyncUI/UnsyncUI/MainWindow.xaml
@@ -7,7 +7,7 @@
         xmlns:mc="http://schemas.openxmlformats.org/markup-compatibility/2006"
         xmlns:local="clr-namespace:UnsyncUI"
         mc:Ignorable="d"
-        Title="UnsyncUI v1.0.8"
+        Title="UnsyncUI v1.0.9"
         Height="800"
         Width="1600"
         MinWidth="600"
