UnsyncUI additions:
Add a Flavor attribute that can be used for different types of package or build configuration. (e.g. with embedded data rather than content-on-demand)
Remove unused Build struct
Modify "Platforms" (could be renamed in future) control to show Platform or Platform + Flavor for discovered items.
Modify Platforms view to sync the size of all checkboxes by wrapping each one in a Grid.

diff --git a/Engine/Source/Programs/UnsyncUI/UnsyncUI/BuildSelectorControl.xaml b/Engine/Source/Programs/UnsyncUI/UnsyncUI/BuildSelectorControl.xaml
--- a/Engine/Source/Programs/UnsyncUI/UnsyncUI/BuildSelectorControl.xaml
+++ b/Engine/Source/Programs/UnsyncUI/UnsyncUI/BuildSelectorControl.xaml
@@ -108,13 +108,18 @@
 				<ListView Grid.Row="1" Foreground="Transparent" Background="Transparent" BorderBrush="Transparent" ItemsSource="{Binding Platforms}" Name="SelectedPlatforms">
 					<ListView.ItemTemplate>
 						<DataTemplate>
-							<CheckBox Margin="4,0,4,0" Content="{Binding Platform}" IsChecked="{Binding Path=IsSelected, RelativeSource={RelativeSource FindAncestor, AncestorType={x:Type ListViewItem}}}"/>
+							<Grid>
+								<Grid.ColumnDefinitions>
+									<ColumnDefinition SharedSizeGroup="WrapPanelGroup" />
+								</Grid.ColumnDefinitions>
+								<CheckBox Margin="4,0,4,0" Content="{Binding Description}" IsChecked="{Binding Path=IsSelected, RelativeSource={RelativeSource FindAncestor, AncestorType={x:Type ListViewItem}}}"/>
+							</Grid>
 						</DataTemplate>
 					</ListView.ItemTemplate>
 					<ListView.ItemsPanel>
 						<ItemsPanelTemplate>
-							<WrapPanel Orientation="Horizontal" 
-									   Width="{Binding RelativeSource={RelativeSource FindAncestor,AncestorType={x:Type ScrollViewer}},Path=ActualWidth}"/>
+							<WrapPanel Orientation="Horizontal" Grid.IsSharedSizeScope="True"
+									   Width="{Binding RelativeSource={RelativeSource FindAncestor,AncestorType={x:Type ScrollViewer}},Path=ActualWidth}" />
 						</ItemsPanelTemplate>
 					</ListView.ItemsPanel>
 					<ListView.ItemContainerStyle>
diff --git a/Engine/Source/Programs/UnsyncUI/UnsyncUI/BuildSelectorControl.xaml.cs b/Engine/Source/Programs/UnsyncUI/UnsyncUI/BuildSelectorControl.xaml.cs
--- a/Engine/Source/Programs/UnsyncUI/UnsyncUI/BuildSelectorControl.xaml.cs
+++ b/Engine/Source/Programs/UnsyncUI/UnsyncUI/BuildSelectorControl.xaml.cs
@@ -81,20 +81,13 @@ namespace UnsyncUI
 		}
 	}
 
-	public struct Build
-	{
-		public string Stream;
-		public string Platform;
-		public string CL;
-		public string Suffix;
-		public string FullPath;
-	}
-
 	public sealed class BuildPlatformModel
 	{
 
 		public BuildModel Build { get; }
 		public string Platform { get; }
+		public string Flavor { get; }
+		public string Description { get; }
 		public string FullPath { get; }
 
 		public string Name => Build == null
@@ -103,11 +96,14 @@ namespace UnsyncUI
 
 		public string SafeName => Name.Replace('/', '+').Replace(" ", "");
 
-		public BuildPlatformModel(BuildModel build, string platform, string fullPath)
+		public BuildPlatformModel(BuildModel build, string platform, string fullPath, string flavor)
 		{
 			Build = build;
 			Platform = platform;
 			FullPath = fullPath;
+			Flavor = flavor;
+
+			Description = Flavor == null ? Platform : $"{Platform} {Flavor}";
 		}
 	}
 
@@ -154,7 +150,7 @@ namespace UnsyncUI
 				if (System.IO.Path.GetFileName(childDir) == ".unsync")
 				{
 					// This folder is a valid build
-					Platforms.Add(new BuildPlatformModel(this, template.Platform, path));
+					Platforms.Add(new BuildPlatformModel(this, template.Platform, path, template.Flavor));
 				}
 				else
 				{
diff --git a/Engine/Source/Programs/UnsyncUI/UnsyncUI/Config.cs b/Engine/Source/Programs/UnsyncUI/UnsyncUI/Config.cs
--- a/Engine/Source/Programs/UnsyncUI/UnsyncUI/Config.cs
+++ b/Engine/Source/Programs/UnsyncUI/UnsyncUI/Config.cs
@@ -38,6 +38,7 @@ namespace UnsyncUI
 			public string CL;
 			public string Suffix;
 			public string Platform;
+			public string Flavor;
 
 			public bool IsStreamClFound => Stream != null && CL != null;
 		}
@@ -49,6 +50,7 @@ namespace UnsyncUI
 			public string Suffix { get; }
 			public string CL { get; }
 			public string Platform { get; }
+			public string Flavor { get; }
 			public List<Directory> SubDirectories { get; }
 
 			public Directory(XElement node)
@@ -58,6 +60,7 @@ namespace UnsyncUI
 				Suffix = node.Attribute("suffix")?.Value;
 				CL = node.Attribute("cl")?.Value;
 				Platform = node.Attribute("platform")?.Value;
+				Flavor = node.Attribute("flavor")?.Value;
 				SubDirectories = node.Elements("dir").Select(d => new Directory(d)).ToList();
 			}
 
@@ -85,6 +88,10 @@ namespace UnsyncUI
 				{
 					template.Platform = RegexReplace(Platform);
 				}
+				if (Flavor != null)
+				{
+					template.Flavor = RegexReplace(Flavor);
+				}
 
 				return true;
 			}
diff --git a/Engine/Source/Programs/UnsyncUI/UnsyncUI/MainWindow.xaml.cs b/Engine/Source/Programs/UnsyncUI/UnsyncUI/MainWindow.xaml.cs
--- a/Engine/Source/Programs/UnsyncUI/UnsyncUI/MainWindow.xaml.cs
+++ b/Engine/Source/Programs/UnsyncUI/UnsyncUI/MainWindow.xaml.cs
@@ -103,7 +103,7 @@ namespace UnsyncUI
 			{
 				onBuildsSelected(new[]
 				{
-					(DstPath, default(string[]), new BuildPlatformModel(null, null, SrcPath))
+					(DstPath, default(string[]), new BuildPlatformModel(null, null, SrcPath, null))
 				});
 			});
 
