Add build Flavor to relative download directory if set.

diff --git a/Engine/Source/Programs/UnsyncUI/UnsyncUI/BuildSelectorControl.xaml.cs b/Engine/Source/Programs/UnsyncUI/UnsyncUI/BuildSelectorControl.xaml.cs
--- a/Engine/Source/Programs/UnsyncUI/UnsyncUI/BuildSelectorControl.xaml.cs
+++ b/Engine/Source/Programs/UnsyncUI/UnsyncUI/BuildSelectorControl.xaml.cs
@@ -89,6 +89,7 @@ namespace UnsyncUI
 		public string Flavor { get; }
 		public string Description { get; }
 		public string FullPath { get; }
+		public string DestPathRelative { get; init; }
 
 		public string Name => Build == null
 			? "Custom Sync" 
@@ -102,6 +103,7 @@ namespace UnsyncUI
 			Platform = platform;
 			FullPath = fullPath;
 			Flavor = flavor;
+			DestPathRelative = Flavor == null ? Platform : Platform + Flavor;
 
 			Description = Flavor == null ? Platform : $"{Platform} {Flavor}";
 		}
@@ -378,7 +380,7 @@ namespace UnsyncUI
 
 		public void StartSync(IEnumerable<BuildPlatformModel> selectedBuilds)
 		{
-			onBuildsSelected?.Invoke(selectedBuilds.Select(s => (Path.Combine(dstPath, s.Platform), Definition.Exclusions?.ToArray(), s)).ToList());
+			onBuildsSelected?.Invoke(selectedBuilds.Select(s => (Path.Combine(dstPath, s.DestPathRelative), Definition.Exclusions?.ToArray(), s)).ToList());
 		}
 
 		public override void OnSelected()
