Ignore new PopulatePlatform requests if we already have one in progress

diff --git a/Engine/Source/Programs/UnsyncUI/UnsyncUI/BuildSelectorControl.xaml.cs b/Engine/Source/Programs/UnsyncUI/UnsyncUI/BuildSelectorControl.xaml.cs
--- a/Engine/Source/Programs/UnsyncUI/UnsyncUI/BuildSelectorControl.xaml.cs
+++ b/Engine/Source/Programs/UnsyncUI/UnsyncUI/BuildSelectorControl.xaml.cs
@@ -142,12 +142,12 @@
 
 		public bool HasPlatform(string platform) => Platforms.Where(p => p.Platform == platform).Any();
 
-		private async Task EnumeratePlatforms(Config.Directory currentDir, string path, Config.BuildTemplate template)
+		private async Task EnumeratePlatforms(Config.Directory currentDir, string path, Config.BuildTemplate template, CancellationTokenSource cancellationToken)
 		{
 			if (!currentDir.Parse(path, ref template))
 				return;
 
-			foreach (var childDir in await AsyncIO.EnumerateDirectoriesAsync(path, cts.Token))
+			foreach (var childDir in await AsyncIO.EnumerateDirectoriesAsync(path, cancellationToken.Token))
 			{
 				if (System.IO.Path.GetFileName(childDir) == ".unsync")
 				{
@@ -158,7 +158,7 @@
 				{
 					foreach (var subDir in currentDir.SubDirectories)
 					{
-						await EnumeratePlatforms(subDir, childDir, template);
+						await EnumeratePlatforms(subDir, childDir, template, cancellationToken);
 					}
 				}
 			}
@@ -166,6 +166,7 @@
 
 		public async void PopulatePlatforms()
 		{
+			cts?.Cancel();
 			cts = new CancellationTokenSource();
 
 			OnRefreshPlatformsClicked.Enabled = false;
@@ -175,7 +176,7 @@
 
 			try
 			{
-				await EnumeratePlatforms(RootDir, Path, Template);
+				await EnumeratePlatforms(RootDir, Path, Template, cts);
 			}
 			catch (OperationCanceledException)
 			{ }
