Add support for syncing subsets of files in a directory in UnsyncUI.
File group is activated by finding any file that matches the regex, then groups from the regex can be used to populate the --include argument to Unsync CLI

diff --git a/Engine/Source/Programs/UnsyncUI/UnsyncUI/App.xaml.cs b/Engine/Source/Programs/UnsyncUI/UnsyncUI/App.xaml.cs
--- a/Engine/Source/Programs/UnsyncUI/UnsyncUI/App.xaml.cs
+++ b/Engine/Source/Programs/UnsyncUI/UnsyncUI/App.xaml.cs
@@ -115,6 +115,7 @@ namespace UnsyncUI
 		public Dictionary<string, string> ProjectDestinationMap { get; set; } = new Dictionary<string, string>();
 		public string CustomSrcPath { get; set; }
 		public string CustomDstPath { get; set; }
+		public string CustomInclude { get; set; }
 		public string AdditionalArgs { get; set; }
 
 		public static UserPreferences Load()
diff --git a/Engine/Source/Programs/UnsyncUI/UnsyncUI/BuildSelectorControl.xaml.cs b/Engine/Source/Programs/UnsyncUI/UnsyncUI/BuildSelectorControl.xaml.cs
--- a/Engine/Source/Programs/UnsyncUI/UnsyncUI/BuildSelectorControl.xaml.cs
+++ b/Engine/Source/Programs/UnsyncUI/UnsyncUI/BuildSelectorControl.xaml.cs
@@ -90,6 +90,7 @@ namespace UnsyncUI
 		public string Description { get; }
 		public string FullPath { get; }
 		public string DestPathRelative { get; init; }
+		public string Include { get;  }
 
 		public string Name => Build == null
 			? "Custom Sync" 
@@ -97,7 +98,7 @@ namespace UnsyncUI
 
 		public string SafeName => Name.Replace('/', '+').Replace(" ", "");
 
-		public BuildPlatformModel(BuildModel build, string platform, string fullPath, string flavor)
+		public BuildPlatformModel(BuildModel build, string platform, string fullPath, string flavor, string include)
 		{
 			Build = build;
 			Platform = platform;
@@ -106,6 +107,7 @@ namespace UnsyncUI
 			DestPathRelative = Flavor == null ? Platform : Platform + Flavor;
 
 			Description = Flavor == null ? Platform : $"{Platform} {Flavor}";
+			Include = include;
 		}
 	}
 
@@ -151,8 +153,25 @@ namespace UnsyncUI
 			{
 				if (System.IO.Path.GetFileName(childDir) == ".unsync")
 				{
-					// This folder is a valid build
-					Platforms.Add(new BuildPlatformModel(this, template.Platform, path, template.Flavor));
+					List<Config.BuildTemplate> fileGroups = null;
+					if (currentDir.FileGroups?.Any() == true)
+					{
+						var files = await AsyncIO.EnumerateFilesAsync(path, cancellationToken.Token);
+						currentDir.ParseFileGroups(files.ToList(), template, out fileGroups);
+					}
+					if (fileGroups?.Any() == true)
+					{
+						foreach (var fileGroup in fileGroups)
+						{
+							Platforms.Add(new BuildPlatformModel(this, fileGroup.Platform, path, fileGroup.Flavor, fileGroup.Include));
+						}
+						Platforms.Add(new BuildPlatformModel(this, template.Platform, path, "All", null));
+					}
+					else
+					{
+						// This folder is a valid build
+						Platforms.Add(new BuildPlatformModel(this, template.Platform, path, template.Flavor, null));
+					}
 				}
 				else
 				{
diff --git a/Engine/Source/Programs/UnsyncUI/UnsyncUI/Config.cs b/Engine/Source/Programs/UnsyncUI/UnsyncUI/Config.cs
--- a/Engine/Source/Programs/UnsyncUI/UnsyncUI/Config.cs
+++ b/Engine/Source/Programs/UnsyncUI/UnsyncUI/Config.cs
@@ -39,8 +39,24 @@ namespace UnsyncUI
 			public string Suffix;
 			public string Platform;
 			public string Flavor;
+			public string Include;
 
 			public bool IsStreamClFound => Stream != null && CL != null;
+
+		}
+
+		public sealed class FileGroup
+		{
+			public Regex Regex { get; }
+			public string Include { get; }
+			public string Flavor { get; }
+
+			public FileGroup(XElement node)
+			{
+				Regex = new Regex($@"^{node.Attribute("regex")?.Value}$", RegexOptions.IgnoreCase);
+				Include = node.Attribute("include")?.Value;
+				Flavor = node.Attribute("flavor")?.Value;
+			}
 		}
 
 		public sealed class Directory
@@ -52,6 +68,7 @@ namespace UnsyncUI
 			public string Platform { get; }
 			public string Flavor { get; }
 			public List<Directory> SubDirectories { get; }
+			public List<FileGroup> FileGroups { get; }
 
 			public Directory(XElement node)
 			{
@@ -62,13 +79,16 @@ namespace UnsyncUI
 				Platform = node.Attribute("platform")?.Value;
 				Flavor = node.Attribute("flavor")?.Value;
 				SubDirectories = node.Elements("dir").Select(d => new Directory(d)).ToList();
+				FileGroups = node.Elements("files").Select(f => new FileGroup(f)).ToList();
 			}
 
 			public bool Parse(string path, ref BuildTemplate template)
 			{
 				var match = Regex.Match(Path.GetFileName(path));
 				if (!match.Success)
+				{
 					return false;
+				}
 
 				string RegexReplace(string field) => Regex.Replace(field, @"\$([0-9]+)", m => match.Groups[int.Parse(m.Groups[1].Value)].Value);
 
@@ -95,6 +115,41 @@ namespace UnsyncUI
 
 				return true;
 			}
+
+			public void ParseFileGroups(List<String> filePaths, BuildTemplate template, out List<BuildTemplate> groupTemplates)
+			{
+				groupTemplates = null;
+				if (!filePaths.Any())
+				{
+					return;
+				}
+
+				HashSet<BuildTemplate> newTemplates = new HashSet<BuildTemplate>();
+				foreach (string path in filePaths)
+				{
+					foreach (var fileGroup in FileGroups)
+					{
+						var match = fileGroup.Regex.Match(Path.GetFileName(path));
+						if (match?.Success == true)
+						{
+							string RegexReplace(string field) => Regex.Replace(field, @"\$([0-9]+)", m => match.Groups[int.Parse(m.Groups[1].Value)].Value);
+
+							BuildTemplate groupTemplate = template;
+							if (fileGroup.Flavor != null)
+							{
+								groupTemplate.Flavor = RegexReplace(fileGroup.Flavor);
+							}
+							if (fileGroup.Include != null)
+							{
+								groupTemplate.Include = RegexReplace(fileGroup.Include);
+							}
+							newTemplates.Add(groupTemplate);
+						}
+					}
+				}
+
+				groupTemplates = newTemplates.ToList();
+			}
 		}
 
 		public sealed class Project
diff --git a/Engine/Source/Programs/UnsyncUI/UnsyncUI/JobControl.xaml.cs b/Engine/Source/Programs/UnsyncUI/UnsyncUI/JobControl.xaml.cs
--- a/Engine/Source/Programs/UnsyncUI/UnsyncUI/JobControl.xaml.cs
+++ b/Engine/Source/Programs/UnsyncUI/UnsyncUI/JobControl.xaml.cs
@@ -3,6 +3,7 @@
 using System;
 using System.Collections.Generic;
 using System.Collections.ObjectModel;
+using System.Diagnostics;
 using System.IO;
 using System.Text;
 using System.Text.RegularExpressions;
@@ -66,6 +67,7 @@ namespace UnsyncUI
         public string Proxy { get; }
 		public string DFS { get; }
 		public string AdditionalArgs { get; }
+		public string[] IncludeFilter { get; }
 		public string[] Exclusions { get; }
 
         public string ProxyStr => string.IsNullOrWhiteSpace(Proxy) ? "(none)" : Proxy;
@@ -162,14 +164,25 @@ namespace UnsyncUI
 		public ObservableCollection<FileModel> Files { get; } = new ObservableCollection<FileModel>();
 		private Dictionary<string, FileModel> filesMap = new Dictionary<string, FileModel>();
 
-        public JobModel(BuildPlatformModel build, string dstPath, bool dryRun, string proxy, string dfs, string additionalArgs, string[] exclusions, Action<JobModel> onCompletion, Action<JobModel> onClear, Action<JobModel> onProgress)
-        {
+		public JobModel(
+			BuildPlatformModel build,
+			string dstPath,
+			bool dryRun,
+			string proxy,
+			string dfs,
+			string additionalArgs,
+			string[] exclusions,
+			Action<JobModel> onCompletion,
+			Action<JobModel> onClear,
+			Action<JobModel> onProgress)
+		{
             Build = build;
             DstPath = dstPath;
             DryRun = dryRun;
             Proxy = proxy;
 			DFS = dfs;
 			AdditionalArgs = additionalArgs;
+			IncludeFilter = build.Include?.Split(',', StringSplitOptions.RemoveEmptyEntries | StringSplitOptions.TrimEntries);
 			Exclusions = exclusions;
 			Succeeded = false;
 
@@ -350,6 +363,10 @@ namespace UnsyncUI
 					args.Add($"--dfs {DFS}");
 				}
 
+				if (IncludeFilter?.Length > 0)
+				{
+					args.Add($"--include {string.Join(",", IncludeFilter)}");
+				}
 				if (Exclusions != null && Exclusions.Length > 0)
 				{
 					args.Add($"--exclude {string.Join(",", Exclusions)}");
@@ -366,6 +383,8 @@ namespace UnsyncUI
 					argsStr = AdditionalArgs.Substring(1);
 				}
 
+				Debug.WriteLine($"Running unsync with args {argsStr}");
+
 				var unsyncPath = App.Current.UnsyncPath;
 				proc = new AsyncProcess(unsyncPath, argsStr);
 
diff --git a/Engine/Source/Programs/UnsyncUI/UnsyncUI/MainWindow.xaml b/Engine/Source/Programs/UnsyncUI/UnsyncUI/MainWindow.xaml
--- a/Engine/Source/Programs/UnsyncUI/UnsyncUI/MainWindow.xaml
+++ b/Engine/Source/Programs/UnsyncUI/UnsyncUI/MainWindow.xaml
@@ -7,7 +7,7 @@
         xmlns:mc="http://schemas.openxmlformats.org/markup-compatibility/2006"
         xmlns:local="clr-namespace:UnsyncUI"
         mc:Ignorable="d"
-        Title="UnsyncUI v1.0.11"
+        Title="UnsyncUI v1.0.12"
         Height="800"
         Width="1600"
         MinWidth="600"
@@ -41,6 +41,7 @@
 							<RowDefinition Height="auto"/>
 							<RowDefinition Height="auto"/>
 							<RowDefinition Height="auto"/>
+							<RowDefinition Height="auto"/>
 							<RowDefinition Height="*"/>
 						</Grid.RowDefinitions>
 
@@ -52,7 +53,11 @@
 							<local:SelectFolderControl SelectedPath="{Binding DstPath}" Description="Select a destination directory to sync builds to."/>
 						</GroupBox>
 
-						<TextBlock Grid.Row="2" Margin="4" HorizontalAlignment="Center">
+						<GroupBox Grid.Row="2" Header="Include">
+							<TextBox Text="{Binding Include}" ToolTip="Comma-separated list parts of filenames to filter the download"/>
+						</GroupBox>
+
+						<TextBlock Grid.Row="3" Margin="4" HorizontalAlignment="Center">
 							<Hyperlink Command="{Binding OnSyncClicked}">Start Custom Sync</Hyperlink>
 						</TextBlock>
 
diff --git a/Engine/Source/Programs/UnsyncUI/UnsyncUI/MainWindow.xaml.cs b/Engine/Source/Programs/UnsyncUI/UnsyncUI/MainWindow.xaml.cs
--- a/Engine/Source/Programs/UnsyncUI/UnsyncUI/MainWindow.xaml.cs
+++ b/Engine/Source/Programs/UnsyncUI/UnsyncUI/MainWindow.xaml.cs
@@ -91,6 +91,17 @@ namespace UnsyncUI
 			}
 		}
 
+		private string include = App.Current.UserConfig.CustomInclude;
+		public string Include
+		{
+			get => include;
+			set
+			{
+				SetProperty(ref include, value);
+				App.Current.UserConfig.CustomInclude = value;
+			}
+		}
+
 		public Command OnSyncClicked { get; }
 
 		private void UpdateSyncCommand()
@@ -103,7 +114,7 @@ namespace UnsyncUI
 			{
 				onBuildsSelected(new[]
 				{
-					(DstPath, default(string[]), new BuildPlatformModel(null, null, SrcPath, null))
+					(DstPath, default(string[]), new BuildPlatformModel(null, null, SrcPath, null, Include))
 				});
 			});
 
@@ -241,7 +252,18 @@ namespace UnsyncUI
 		{
 			foreach (var build in selectedBuilds)
 			{
-				AddJob(new JobModel(build.Model, build.DstPath, DryRun, SelectedProxy?.Path, DFS, AdditionalArgs, build.Exclusions, OnJobCompleted, OnClearJob, OnJobProgress));
+				AddJob(new JobModel(
+					build.Model, 
+					build.DstPath, 
+					DryRun, 
+					SelectedProxy?.Path,
+					DFS, 
+					AdditionalArgs, 
+					build.Exclusions, 
+					OnJobCompleted, 
+					OnClearJob, 
+					OnJobProgress
+				));
 			}
 		}
 
diff --git a/Engine/Source/Programs/UnsyncUI/UnsyncUI/Utils.cs b/Engine/Source/Programs/UnsyncUI/UnsyncUI/Utils.cs
--- a/Engine/Source/Programs/UnsyncUI/UnsyncUI/Utils.cs
+++ b/Engine/Source/Programs/UnsyncUI/UnsyncUI/Utils.cs
@@ -108,6 +108,45 @@ namespace UnsyncUI
 				}
 			});
 
+            return tcs.Task;
+		}
+        public static Task<IEnumerable<string>> EnumerateFilesAsync(string path, CancellationToken token)
+        {
+			var tcs = new TaskCompletionSource<IEnumerable<string>>();
+            Task.Run(() =>
+			{
+				using var cancel = token.Register(() => tcs.TrySetCanceled());
+				try
+				{
+					mutex.Wait(token);
+					try
+					{
+						var timer = new Stopwatch();
+						timer.Start();
+
+						var dirs = Directory.EnumerateFiles(path).ToList();
+
+						timer.Stop();
+						Debug.WriteLine($"Time: {timer.Elapsed.TotalSeconds:0.000} s - {path}");
+
+						tcs.TrySetResult(dirs);
+					}
+					catch (Exception ex)
+					{
+						tcs.TrySetException(ex);
+					}
+					finally
+					{
+						mutex.Release();
+					}
+				}
+				catch (OperationCanceledException)
+				{
+					// Mutex wait was cancelled. Don't release
+					tcs.TrySetCanceled();
+				}
+			});
+
             return tcs.Task;
 		}
     }
