UnsyncUI - Use .Net Core 3.1 (same as UGS)

diff --git a/Engine/Source/Programs/UnsyncUI/UnsyncUI/BuildSelectorControl.xaml.cs b/Engine/Source/Programs/UnsyncUI/UnsyncUI/BuildSelectorControl.xaml.cs
--- a/Engine/Source/Programs/UnsyncUI/UnsyncUI/BuildSelectorControl.xaml.cs
+++ b/Engine/Source/Programs/UnsyncUI/UnsyncUI/BuildSelectorControl.xaml.cs
@@ -89,7 +89,7 @@ namespace UnsyncUI
 		public string Flavor { get; }
 		public string Description { get; }
 		public string FullPath { get; }
-		public string DestPathRelative { get; init; }
+		public string DestPathRelative { get; }
 		public string Include { get;  }
 
 		public string Name => Build == null
diff --git a/Engine/Source/Programs/UnsyncUI/UnsyncUI/JobControl.xaml.cs b/Engine/Source/Programs/UnsyncUI/UnsyncUI/JobControl.xaml.cs
--- a/Engine/Source/Programs/UnsyncUI/UnsyncUI/JobControl.xaml.cs
+++ b/Engine/Source/Programs/UnsyncUI/UnsyncUI/JobControl.xaml.cs
@@ -5,6 +5,7 @@ using System.Collections.Generic;
 using System.Collections.ObjectModel;
 using System.Diagnostics;
 using System.IO;
+using System.Linq;
 using System.Text;
 using System.Text.RegularExpressions;
 using System.Threading;
@@ -145,7 +146,7 @@ namespace UnsyncUI
 		public Command CopyLogClicked { get; }
 		public Command SaveLogClicked { get; }
 
-		private TaskCompletionSource startTcs = new TaskCompletionSource();
+		private TaskCompletionSource<bool> startTcs = new TaskCompletionSource<bool>();
         private CancellationTokenSource cts = new CancellationTokenSource();
         public bool IsCancelled => cts.IsCancellationRequested;
 		public bool Succeeded { get; private set; }
@@ -182,7 +183,10 @@ namespace UnsyncUI
             Proxy = proxy;
 			DFS = dfs;
 			AdditionalArgs = additionalArgs;
-			IncludeFilter = build.Include?.Split(',', StringSplitOptions.RemoveEmptyEntries | StringSplitOptions.TrimEntries);
+			IncludeFilter = build.Include?.Split(',', StringSplitOptions.RemoveEmptyEntries)
+				.Select(Entry => Entry.Trim())
+				.Where(Entry => Entry.Length > 0)
+				.ToArray();
 			Exclusions = exclusions;
 			Succeeded = false;
 
@@ -242,7 +246,7 @@ namespace UnsyncUI
 			}
 		}
 
-        public void StartJob() => startTcs.TrySetResult();
+        public void StartJob() => startTcs.TrySetResult(true);
 
 		public void CancelJob()
 		{
diff --git a/Engine/Source/Programs/UnsyncUI/UnsyncUI/UnsyncUI.csproj b/Engine/Source/Programs/UnsyncUI/UnsyncUI/UnsyncUI.csproj
--- a/Engine/Source/Programs/UnsyncUI/UnsyncUI/UnsyncUI.csproj
+++ b/Engine/Source/Programs/UnsyncUI/UnsyncUI/UnsyncUI.csproj
@@ -2,7 +2,7 @@
 
 	<PropertyGroup>
 		<OutputType>WinExe</OutputType>
-		<TargetFramework>net5.0-windows</TargetFramework>
+		<TargetFramework>netcoreapp3.1</TargetFramework>
 		<UseWPF>true</UseWPF>
 		<UseWindowsForms>true</UseWindowsForms>
 		<ApplicationManifest>app.manifest</ApplicationManifest>
