UnsyncUI - Report errors in the build list area if any were encountered during enumeration

* Bump version to 1.0.13

diff --git a/Engine/Source/Programs/UnsyncUI/UnsyncUI/BuildSelectorControl.xaml b/Engine/Source/Programs/UnsyncUI/UnsyncUI/BuildSelectorControl.xaml
--- a/Engine/Source/Programs/UnsyncUI/UnsyncUI/BuildSelectorControl.xaml
+++ b/Engine/Source/Programs/UnsyncUI/UnsyncUI/BuildSelectorControl.xaml
@@ -71,6 +71,10 @@
 			</Grid>
 		</GroupBox>
 
+		<TextBlock TextAlignment="Center" VerticalAlignment="Center" Visibility="{Binding Path=StatusString, Converter={StaticResource VisConv_Collapsed}}">
+			<Run Text="{Binding Path=StatusString, Mode=OneWay}"/>
+		</TextBlock>
+
 		<!-- Platforms Groupbox -->
 		<GroupBox Grid.Row="1" MinHeight="62" DataContext="{Binding SelectedItem, ElementName=SelectedBuild}" Visibility="{Binding SelectedItem, ElementName=SelectedBuild, Converter={StaticResource VisConv_Collapsed}}">
 			
@@ -142,7 +146,7 @@
 		<GroupBox Grid.Row="2">
 			<GroupBox.Header>
 				<Grid>
-					<TextBlock Margin="4" Visibility="{Binding Path=!ResetDestinationClicked.Enabled, Converter={StaticResource VisConv_Collapsed}}">
+					<TextBlock Margin="4" Visibility="{Binding Path=ResetDestinationClicked.Enabled, Converter={StaticResource VisConv_Collapsed_Inv}}">
 						<Run Text="Destination"/>
 					</TextBlock>
 					<TextBlock Margin="4" Visibility="{Binding Path=ResetDestinationClicked.Enabled, Converter={StaticResource VisConv_Collapsed}}">
diff --git a/Engine/Source/Programs/UnsyncUI/UnsyncUI/BuildSelectorControl.xaml.cs b/Engine/Source/Programs/UnsyncUI/UnsyncUI/BuildSelectorControl.xaml.cs
--- a/Engine/Source/Programs/UnsyncUI/UnsyncUI/BuildSelectorControl.xaml.cs
+++ b/Engine/Source/Programs/UnsyncUI/UnsyncUI/BuildSelectorControl.xaml.cs
@@ -212,6 +212,16 @@ namespace UnsyncUI
 		private bool firstRefresh = true;
 		private Action<IEnumerable<(string DstPath, string[] Exclusions, BuildPlatformModel Model)>> onBuildsSelected;
 
+		private string statusString = null;
+		public string StatusString
+		{
+			get => statusString;
+			set
+			{
+				SetProperty(ref statusString, value);
+			}
+		}
+
 		public Config.Project Definition { get; }
 		public string Name => Definition.Name;
 
@@ -370,12 +380,13 @@ namespace UnsyncUI
 
 			Builds.Clear();
 			SelectedBuild = null;
+			StatusString = null;
 
 			try
 			{
 				var timer = new Stopwatch();
 				timer.Start();
-				var buildsPipe = Definition.EnumerateBuilds(cts.Token);
+				var (task, buildsPipe) = Definition.EnumerateBuilds(cts.Token);
 
 				while (await buildsPipe.OutputAvailableAsync(cts.Token))
 				{
@@ -383,14 +394,23 @@ namespace UnsyncUI
 				}
 
 				timer.Stop();
-				Debug.WriteLine($"Search took {timer.Elapsed.TotalSeconds} seconds.");
-			}
-			catch (IOException)
-			{
-				// @todo show an error message
+
+				if (task.IsFaulted)
+				{
+					throw task.Exception.InnerException;
+				}
+				else
+				{
+					Debug.WriteLine($"Search took {timer.Elapsed.TotalSeconds} seconds.");
+				}
 			}
 			catch (OperationCanceledException)
 			{ }
+			catch (Exception ex)
+			{
+				Debug.WriteLine($"Directory enumeration failed with exception: {ex}");
+				StatusString = ex.Message;
+			}
 			finally
 			{
 				OnRefreshBuildsClicked.Enabled = true;
diff --git a/Engine/Source/Programs/UnsyncUI/UnsyncUI/Config.cs b/Engine/Source/Programs/UnsyncUI/UnsyncUI/Config.cs
--- a/Engine/Source/Programs/UnsyncUI/UnsyncUI/Config.cs
+++ b/Engine/Source/Programs/UnsyncUI/UnsyncUI/Config.cs
@@ -184,10 +184,10 @@ namespace UnsyncUI
 				});
 			}
 
-			public ISourceBlock<BuildModel> EnumerateBuilds(CancellationToken cancelToken)
+			public (Task, ISourceBlock<BuildModel>) EnumerateBuilds(CancellationToken cancelToken)
 			{
 				var pipe = new BufferBlock<BuildModel>();
-				Task.Run(async () =>
+				var task = Task.Run(async () =>
 				{
 					try
 					{
@@ -205,7 +205,7 @@ namespace UnsyncUI
 					}
 				});
 
-				return pipe;
+				return (task, pipe);
 			}
 		}
 
diff --git a/Engine/Source/Programs/UnsyncUI/UnsyncUI/MainWindow.xaml b/Engine/Source/Programs/UnsyncUI/UnsyncUI/MainWindow.xaml
--- a/Engine/Source/Programs/UnsyncUI/UnsyncUI/MainWindow.xaml
+++ b/Engine/Source/Programs/UnsyncUI/UnsyncUI/MainWindow.xaml
@@ -7,7 +7,7 @@
         xmlns:mc="http://schemas.openxmlformats.org/markup-compatibility/2006"
         xmlns:local="clr-namespace:UnsyncUI"
         mc:Ignorable="d"
-        Title="UnsyncUI v1.0.12"
+        Title="UnsyncUI"
         Height="800"
         Width="1600"
         MinWidth="600"
diff --git a/Engine/Source/Programs/UnsyncUI/UnsyncUI/MainWindow.xaml.cs b/Engine/Source/Programs/UnsyncUI/UnsyncUI/MainWindow.xaml.cs
--- a/Engine/Source/Programs/UnsyncUI/UnsyncUI/MainWindow.xaml.cs
+++ b/Engine/Source/Programs/UnsyncUI/UnsyncUI/MainWindow.xaml.cs
@@ -21,6 +21,8 @@ namespace UnsyncUI
 		{
 			InitializeComponent();
 			DataContext = new MainWindowModel();
+			var version = System.Reflection.Assembly.GetEntryAssembly().GetName().Version;
+			Title = $"UnsyncUI v{version.Major}.{version.Minor}.{version.Build}";
 		}
 
 		private void Window_Closing(object sender, CancelEventArgs e)
diff --git a/Engine/Source/Programs/UnsyncUI/UnsyncUI/UnsyncUI.csproj b/Engine/Source/Programs/UnsyncUI/UnsyncUI/UnsyncUI.csproj
--- a/Engine/Source/Programs/UnsyncUI/UnsyncUI/UnsyncUI.csproj
+++ b/Engine/Source/Programs/UnsyncUI/UnsyncUI/UnsyncUI.csproj
@@ -8,6 +8,7 @@
 		<ApplicationManifest>app.manifest</ApplicationManifest>
         <IncludeAllContentForSelfExtract>true</IncludeAllContentForSelfExtract>
         <ApplicationIcon>unsync_48px.ico</ApplicationIcon>
+        <Version>$(VersionPrefix)1.0.13</Version>
 	</PropertyGroup>
 
 </Project>
