UnsyncUI :
* Added context menu option to build selector (Copy Build Name : copies the selected build's folder name to the clipboard)
* Added checkbox to append the selected build's name to the destination folder
* Added hyperlink to open the destination folder

diff --git a/Engine/Source/Programs/UnsyncUI/UnsyncUI/App.xaml.cs b/Engine/Source/Programs/UnsyncUI/UnsyncUI/App.xaml.cs
--- a/Engine/Source/Programs/UnsyncUI/UnsyncUI/App.xaml.cs
+++ b/Engine/Source/Programs/UnsyncUI/UnsyncUI/App.xaml.cs
@@ -117,6 +117,7 @@ namespace UnsyncUI
 		public string CustomDstPath { get; set; }
 		public string CustomInclude { get; set; }
 		public string AdditionalArgs { get; set; }
+		public bool AppendBuildName { get; set; } = false;
 
 		public static UserPreferences Load()
 		{
diff --git a/Engine/Source/Programs/UnsyncUI/UnsyncUI/BuildSelectorControl.xaml b/Engine/Source/Programs/UnsyncUI/UnsyncUI/BuildSelectorControl.xaml
--- a/Engine/Source/Programs/UnsyncUI/UnsyncUI/BuildSelectorControl.xaml
+++ b/Engine/Source/Programs/UnsyncUI/UnsyncUI/BuildSelectorControl.xaml
@@ -57,6 +57,7 @@
 					<ListView.ContextMenu>
 						<ContextMenu>
 							<MenuItem Header="Copy Build CL" Command="{Binding CopyBuildCLClicked}" />
+							<MenuItem Header="Copy Build Name" Command="{Binding CopyBuildNameClicked}" />
 							<MenuItem Header="Open Build Location" Command="{Binding OpenBuildLocationClicked}" />
 						</ContextMenu>
 					</ListView.ContextMenu>
@@ -151,12 +152,29 @@
 					</TextBlock>
 					<TextBlock Margin="4" Visibility="{Binding Path=ResetDestinationClicked.Enabled, Converter={StaticResource VisConv_Collapsed}}">
 						<Run Text="Destination - "/>
-					<Hyperlink Command="{Binding ResetDestinationClicked}">Reset</Hyperlink>
+						<Hyperlink Command="{Binding ResetDestinationClicked}">Reset</Hyperlink>
 					</TextBlock>
 				</Grid>
 			</GroupBox.Header>
 
-			<local:SelectFolderControl SelectedPath="{Binding DstPath, UpdateSourceTrigger=PropertyChanged}" Description="Select a destination directory to sync builds to."/>
+			<Grid>
+				<Grid.RowDefinitions>
+					<RowDefinition Height="auto"/>
+					<RowDefinition Height="auto"/>
+					<RowDefinition Height="auto"/>
+				</Grid.RowDefinitions>
+
+				<local:SelectFolderControl Grid.Row ="0" SelectedPath="{Binding DstPath, UpdateSourceTrigger=PropertyChanged}" Description="Select a destination directory to sync builds to."/>
+
+				<CheckBox Grid.Row ="1" VerticalAlignment="Center" IsChecked="{Binding AppendBuildName}" 
+						  Content="Append Build Name" ToolTip="If a build is selected, its name will be appended to the destination folder." 
+						  ToolTipService.ShowOnDisabled="True" HorizontalAlignment="Left" Margin="4,0,0,0"/>
+				<TextBlock Margin="4" Grid.Row ="2">
+					<Hyperlink Command="{Binding OpenDestinationClicked}">
+						<TextBlock Text="{Binding FinalDstPath, UpdateSourceTrigger=PropertyChanged, Mode=OneWay}"/>
+					</Hyperlink>
+				</TextBlock>
+			</Grid>
 		</GroupBox>
 
 		<!-- Start sync -->
diff --git a/Engine/Source/Programs/UnsyncUI/UnsyncUI/BuildSelectorControl.xaml.cs b/Engine/Source/Programs/UnsyncUI/UnsyncUI/BuildSelectorControl.xaml.cs
--- a/Engine/Source/Programs/UnsyncUI/UnsyncUI/BuildSelectorControl.xaml.cs
+++ b/Engine/Source/Programs/UnsyncUI/UnsyncUI/BuildSelectorControl.xaml.cs
@@ -229,8 +229,10 @@ namespace UnsyncUI
 		public Command OnStopRefreshBuildsClicked { get; }
 		public Command<IList> OnStartSync { get; }
 		public Command CopyBuildCLClicked { get; }
+		public Command CopyBuildNameClicked { get; }
 		public Command OpenBuildLocationClicked { get; }
 		public Command ResetDestinationClicked { get; }
+		public Command OpenDestinationClicked { get; }
 
 		public ObservableCollection<BuildModel> Builds { get; } = new ObservableCollection<BuildModel>();
 		public ICollectionView BuildsView { get; }
@@ -260,6 +262,8 @@ namespace UnsyncUI
 			{
 				SetProperty(ref selectedBuild, value);
 
+				UpdateFinalDstPath();
+
 				if (selectedBuild != null)
 				{
 					selectedBuild.PopulatePlatforms();
@@ -294,12 +298,55 @@ namespace UnsyncUI
 			set
 			{
 				SetProperty(ref dstPath, value);
-				App.Current.UserConfig.ProjectDestinationMap[Name] = value;
+
+				UpdateFinalDstPath();
+
+				App.Current.UserConfig.ProjectDestinationMap[Name] = DstPath;
 
 				ResetDestinationClicked.Enabled = CanResetDestinationToDefault;
 			}
 		}
 
+		void UpdateFinalDstPath()
+		{
+			string finalPath = dstPath;
+			if (AppendBuildName && (selectedBuild != null))
+			{
+				string buildFolderName = new DirectoryInfo(selectedBuild.Path).Name;
+				finalPath = Path.Combine(DstPath, buildFolderName);
+			}
+
+			finalPath = Path.TrimEndingDirectorySeparator(finalPath);
+			if (FinalDstPath != finalPath)
+			{
+				FinalDstPath = finalPath;
+			}
+		}
+
+		private string finalDstPath = null;
+		public string FinalDstPath
+		{
+			get => finalDstPath;
+			set
+			{
+				SetProperty(ref finalDstPath, value);
+			}
+		}
+
+		private bool appendBuildName = true;
+		public bool AppendBuildName
+		{
+			get => appendBuildName;
+			set
+			{
+				SetProperty(ref appendBuildName, value);
+
+				App.Current.UserConfig.AppendBuildName = value;
+
+				UpdateFinalDstPath();
+			}
+		}
+
 		public ProjectModel(Config.Project definition, Action<IEnumerable<(string DstPath, string[] Exclusions, BuildPlatformModel Model)>> onBuildsSelected)
 		{
 			this.Definition = definition;
@@ -308,7 +355,9 @@ namespace UnsyncUI
 			OnStopRefreshBuildsClicked = new Command(StopRefreshBuilds) { Enabled = false };
 			OnStartSync = new Command<IList>(l => StartSync(l.OfType<BuildPlatformModel>())) { Enabled = true };
 			CopyBuildCLClicked = new Command(CopyBuildCL) { Enabled = true };
+			CopyBuildNameClicked = new Command(CopyBuildName) { Enabled = true };
 			OpenBuildLocationClicked = new Command(OpenBuildLocation) { Enabled = true };
+			OpenDestinationClicked = new Command(OpenDestinationLocation) { Enabled = true };
 
 			BuildsView = CollectionViewSource.GetDefaultView(Builds);
 			BuildsView.Filter = BuildsViewFilter;
@@ -342,6 +391,8 @@ namespace UnsyncUI
 				defaultDstPath = "";
 			}
 
+			appendBuildName = App.Current.UserConfig.AppendBuildName;
+
 			string userDstPath = null;
 
 			if (App.Current.UserConfig.ProjectDestinationMap.TryGetValue(Name, out userDstPath))
@@ -352,6 +403,7 @@ namespace UnsyncUI
 			{
 				dstPath = defaultDstPath;
 			}
+			UpdateFinalDstPath();
 
 			ResetDestinationClicked = new Command(ResetDestinationToDefault) { Enabled = CanResetDestinationToDefault };
 		}
@@ -420,7 +472,7 @@ namespace UnsyncUI
 
 		public void StartSync(IEnumerable<BuildPlatformModel> selectedBuilds)
 		{
-			onBuildsSelected?.Invoke(selectedBuilds.Select(s => (Path.Combine(dstPath, s.DestPathRelative), Definition.Exclusions?.ToArray(), s)).ToList());
+			onBuildsSelected?.Invoke(selectedBuilds.Select(s => (Path.Combine(finalDstPath, s.DestPathRelative), Definition.Exclusions?.ToArray(), s)).ToList());
 		}
 
 		public override void OnSelected()
@@ -432,33 +484,47 @@ namespace UnsyncUI
 			}
 		}
 
+		private void CopyStringToClipboard(string str)
+		{
+			// Clipboard can be locked if timing is just right(wrong), retry a few times on failure.
+			int retryDelay = 0;
+			const int attempts = 7;
+			for (int attempt = 0; attempt < attempts; attempt++)
+			{
+				if (retryDelay > 0)
+				{
+					System.Threading.Thread.Sleep(retryDelay);
+					retryDelay *= 2;
+				}
+				else
+				{
+					retryDelay = 5;
+				}
+				try
+				{
+					Clipboard.SetText(str, TextDataFormat.Text);
+					break;
+				}
+				catch (System.Runtime.InteropServices.COMException)
+				{
+				}
+			}
+		}
+
 		private void CopyBuildCL()
 		{
 			if (selectedBuild != null)
 			{
-				// Clipboard can be locked if timing is just right(wrong), retry a few times on failure.
-                int retryDelay = 0;
-				const int attempts = 7;
-				for (int attempt = 0; attempt < attempts; attempt++)
-				{
-					if (retryDelay > 0)
-					{
-                        System.Threading.Thread.Sleep(retryDelay);
-                        retryDelay *= 2;
-                    }
-					else
-					{
-						retryDelay = 5;
-                    }
-                    try
-                    {
-                        Clipboard.SetText(selectedBuild.CL, TextDataFormat.Text);
-						break;
-                    }
-                    catch (System.Runtime.InteropServices.COMException)
-                    {
-	                }
-                }
+				CopyStringToClipboard(selectedBuild.CL);
+			}
+		}
+
+		private void CopyBuildName()
+		{
+			if (selectedBuild != null)
+			{
+				string buildFolderName = new DirectoryInfo(selectedBuild.Path).Name;
+				CopyStringToClipboard(buildFolderName);
 			}
 		}
 
@@ -469,5 +535,20 @@ namespace UnsyncUI
 				Shell.LaunchExplorer(selectedBuild.Path);
 			}
 		}
+
+		private void OpenDestinationLocation()
+		{
+			if (!Directory.Exists(FinalDstPath))
+			{
+				try
+				{
+					Directory.CreateDirectory(FinalDstPath);
+				}
+				catch (Exception)
+				{
+				}
+			}
+			Shell.LaunchExplorer(FinalDstPath);
+		}
 	}
 }
diff --git a/Engine/Source/Programs/UnsyncUI/UnsyncUI/UnsyncUI.csproj b/Engine/Source/Programs/UnsyncUI/UnsyncUI/UnsyncUI.csproj
--- a/Engine/Source/Programs/UnsyncUI/UnsyncUI/UnsyncUI.csproj
+++ b/Engine/Source/Programs/UnsyncUI/UnsyncUI/UnsyncUI.csproj
@@ -8,7 +8,7 @@
 		<ApplicationManifest>app.manifest</ApplicationManifest>
         <IncludeAllContentForSelfExtract>true</IncludeAllContentForSelfExtract>
         <ApplicationIcon>unsync_48px.ico</ApplicationIcon>
-        <Version>$(VersionPrefix)1.0.13</Version>
+        <Version>$(VersionPrefix)1.0.14</Version>
 	</PropertyGroup>
 
 </Project>
