UnsyncUI v1.0.15 :
* Adjusted the Append Build Name tooltip to warn the user about increased sync times
* Bound the Reset button visibility to the final path (so that Append build name is taken into account)
* Uncheck Append build name on Reset
* Disabled Append build name when no build is selected

#skipci

diff --git a/Engine/Source/Programs/UnsyncUI/UnsyncUI/App.xaml b/Engine/Source/Programs/UnsyncUI/UnsyncUI/App.xaml
--- a/Engine/Source/Programs/UnsyncUI/UnsyncUI/App.xaml
+++ b/Engine/Source/Programs/UnsyncUI/UnsyncUI/App.xaml
@@ -35,6 +35,7 @@
 		<local:VisibilityConverter x:Key="VisConv_Collapsed" FalseState="Collapsed" />
 		<local:VisibilityConverter x:Key="VisConv_Collapsed_Inv" TrueState="Collapsed" FalseState="Visible" />
 		<local:VisibilityConverter x:Key="VisConv_Hidden" FalseState="Hidden" />
+		<local:ObjectToBoolConverter x:Key="IsObjectNotNull"/>	
 		
 	</Application.Resources>
 </Application>
diff --git a/Engine/Source/Programs/UnsyncUI/UnsyncUI/BuildSelectorControl.xaml b/Engine/Source/Programs/UnsyncUI/UnsyncUI/BuildSelectorControl.xaml
--- a/Engine/Source/Programs/UnsyncUI/UnsyncUI/BuildSelectorControl.xaml
+++ b/Engine/Source/Programs/UnsyncUI/UnsyncUI/BuildSelectorControl.xaml
@@ -166,8 +166,8 @@
 
 				<local:SelectFolderControl Grid.Row ="0" SelectedPath="{Binding DstPath, UpdateSourceTrigger=PropertyChanged}" Description="Select a destination directory to sync builds to."/>
 
-				<CheckBox Grid.Row ="1" VerticalAlignment="Center" IsChecked="{Binding AppendBuildName}" 
-						  Content="Append Build Name" ToolTip="If a build is selected, its name will be appended to the destination folder." 
+				<CheckBox Grid.Row ="1" VerticalAlignment="Center" IsChecked="{Binding AppendBuildName}" IsEnabled="{Binding SelectedBuild, Converter={StaticResource IsObjectNotNull}}"
+						  Content="Append Build Name" ToolTip="If a build is selected, its name will be appended to the destination folder.&#x0a;Note : using this option will download each build in a different folder, which is useful when you need multiple builds for comparison or whatnot. However, downloading builds over an existing folder yields faster sync times." 
 						  ToolTipService.ShowOnDisabled="True" HorizontalAlignment="Left" Margin="4,0,0,0"/>
 				<TextBlock Margin="4" Grid.Row ="2">
 					<Hyperlink Command="{Binding OpenDestinationClicked}">
diff --git a/Engine/Source/Programs/UnsyncUI/UnsyncUI/BuildSelectorControl.xaml.cs b/Engine/Source/Programs/UnsyncUI/UnsyncUI/BuildSelectorControl.xaml.cs
--- a/Engine/Source/Programs/UnsyncUI/UnsyncUI/BuildSelectorControl.xaml.cs
+++ b/Engine/Source/Programs/UnsyncUI/UnsyncUI/BuildSelectorControl.xaml.cs
@@ -302,8 +302,6 @@ namespace UnsyncUI
 				UpdateFinalDstPath();
 
 				App.Current.UserConfig.ProjectDestinationMap[Name] = DstPath;
-
-				ResetDestinationClicked.Enabled = CanResetDestinationToDefault;
 			}
 		}
 
@@ -321,6 +319,8 @@ namespace UnsyncUI
 			{
 				FinalDstPath = finalPath;
 			}
+
+			ResetDestinationClicked.Enabled = CanResetDestinationToDefault;
 		}
 
 		private string finalDstPath = null;
@@ -393,6 +393,8 @@ namespace UnsyncUI
 
 			appendBuildName = App.Current.UserConfig.AppendBuildName;
 
+			ResetDestinationClicked = new Command(ResetDestinationToDefault) { Enabled = CanResetDestinationToDefault };
+
 			string userDstPath = null;
 
 			if (App.Current.UserConfig.ProjectDestinationMap.TryGetValue(Name, out userDstPath))
@@ -404,14 +406,13 @@ namespace UnsyncUI
 				dstPath = defaultDstPath;
 			}
 			UpdateFinalDstPath();
-
-			ResetDestinationClicked = new Command(ResetDestinationToDefault) { Enabled = CanResetDestinationToDefault };
 		}
-		private bool CanResetDestinationToDefault => dstPath != defaultDstPath && defaultDstPath.Length != 0;
+		private bool CanResetDestinationToDefault => finalDstPath != defaultDstPath && defaultDstPath.Length != 0;
 
 		private void ResetDestinationToDefault()
 		{
 			DstPath = defaultDstPath;
+			AppendBuildName = false;
 		}
 
 		private void StopRefreshBuilds()
diff --git a/Engine/Source/Programs/UnsyncUI/UnsyncUI/Converters.cs b/Engine/Source/Programs/UnsyncUI/UnsyncUI/Converters.cs
--- a/Engine/Source/Programs/UnsyncUI/UnsyncUI/Converters.cs
+++ b/Engine/Source/Programs/UnsyncUI/UnsyncUI/Converters.cs
@@ -134,4 +134,22 @@ namespace UnsyncUI
 		public object ConvertBack(object value, Type targetType, object parameter, CultureInfo culture)
 			=> throw new NotImplementedException();
 	}
+
+	public class ObjectToBoolConverter : IValueConverter
+	{
+		public object Convert(object value, Type targetType, object parameter, CultureInfo culture)
+		{
+			bool returnValue = true;
+			if (value == DependencyProperty.UnsetValue || value == null)
+			{
+				returnValue = false;
+			}
+			return returnValue;
+		}
+
+		public object ConvertBack(object value, Type targetType, object parameter, CultureInfo culture)
+		{
+			return Binding.DoNothing;
+		}
+	}
 }
diff --git a/Engine/Source/Programs/UnsyncUI/UnsyncUI/UnsyncUI.csproj b/Engine/Source/Programs/UnsyncUI/UnsyncUI/UnsyncUI.csproj
--- a/Engine/Source/Programs/UnsyncUI/UnsyncUI/UnsyncUI.csproj
+++ b/Engine/Source/Programs/UnsyncUI/UnsyncUI/UnsyncUI.csproj
@@ -8,7 +8,7 @@
 		<ApplicationManifest>app.manifest</ApplicationManifest>
         <IncludeAllContentForSelfExtract>true</IncludeAllContentForSelfExtract>
         <ApplicationIcon>unsync_48px.ico</ApplicationIcon>
-        <Version>$(VersionPrefix)1.0.14</Version>
+        <Version>$(VersionPrefix)1.0.15</Version>
 	</PropertyGroup>
 
 </Project>
