UnsyncUI v1.0.16 :
* Changed Queued jobs to act as FIFO, so that we end up downloading all platforms from a given build before starting to process another build.
* Avoid enqueuing a job for a build+platform if it's a duplicate of one that's currently in-flight or queued up

#skipci

diff --git a/Engine/Source/Programs/UnsyncUI/UnsyncUI/JobControl.xaml.cs b/Engine/Source/Programs/UnsyncUI/UnsyncUI/JobControl.xaml.cs
--- a/Engine/Source/Programs/UnsyncUI/UnsyncUI/JobControl.xaml.cs
+++ b/Engine/Source/Programs/UnsyncUI/UnsyncUI/JobControl.xaml.cs
@@ -64,6 +64,7 @@ namespace UnsyncUI
     {
         public BuildPlatformModel Build { get; }
         public string DstPath { get; }
+		public string DstPathBase => Path.GetDirectoryName(DstPath);
         public bool DryRun { get; }
         public string Proxy { get; }
 		public string DFS { get; }
@@ -75,7 +76,7 @@ namespace UnsyncUI
 
         public string Name => $"{Build.Name}{(DryRun ? " - Dry Run" : "")}";
 
-        private double totalProgress = 0;
+		private double totalProgress = 0;
         public double TotalProgress
         {
             get => totalProgress;
@@ -258,6 +259,13 @@ namespace UnsyncUI
 			cts.Cancel();
 		}
 
+		public bool IsDuplicate(JobModel otherJob) 
+		{
+			return (Build.Build.Path == otherJob.Build.Build.Path)
+				&& (Build.Platform == otherJob.Build.Platform)
+				&& (DstPathBase == otherJob.DstPathBase);
+		}
+
         private async void RunAsync()
         {
 			AsyncProcess proc = null;
diff --git a/Engine/Source/Programs/UnsyncUI/UnsyncUI/MainWindow.xaml.cs b/Engine/Source/Programs/UnsyncUI/UnsyncUI/MainWindow.xaml.cs
--- a/Engine/Source/Programs/UnsyncUI/UnsyncUI/MainWindow.xaml.cs
+++ b/Engine/Source/Programs/UnsyncUI/UnsyncUI/MainWindow.xaml.cs
@@ -278,9 +278,17 @@ namespace UnsyncUI
 				ActiveJob = job;
 				job.StartJob();
 			}
-			else
+			// Don't bother including a job that is already either active or queued up :
+			else if (!ActiveJob.IsDuplicate(job) && (QueuedJobs.FirstOrDefault(j => j.IsDuplicate(job)) == null))
 			{
-				QueuedJobs.Add(job);
+				// Add it after the last platform for that same build if any, so that jobs from the same build end up being processed one after another 
+				JobModel lastJobForBuild = QueuedJobs.LastOrDefault(j => (j.Build.Build.Path == job.Build.Build.Path) && (j.DstPathBase == job.DstPathBase));
+				int insertionIndex = 0;
+				if (lastJobForBuild != null)
+				{
+					insertionIndex = QueuedJobs.IndexOf(lastJobForBuild) + 1;
+				}
+				QueuedJobs.Insert(insertionIndex, job);
 			}
 
 			UpdateProgressValue();
@@ -311,7 +319,8 @@ namespace UnsyncUI
 				ActiveJob = null;
 				if (QueuedJobs.Count > 0)
 				{
-					var newJob = QueuedJobs.Where(q => !q.IsCancelled).LastOrDefault();
+					// Look for a new job to run. Pick the first one rather than the last one for the queue to act as FIFO : 
+					var newJob = QueuedJobs.Where(q => !q.IsCancelled).FirstOrDefault();
 					if (newJob != null)
 					{
 						QueuedJobs.Remove(newJob);
diff --git a/Engine/Source/Programs/UnsyncUI/UnsyncUI/UnsyncUI.csproj b/Engine/Source/Programs/UnsyncUI/UnsyncUI/UnsyncUI.csproj
--- a/Engine/Source/Programs/UnsyncUI/UnsyncUI/UnsyncUI.csproj
+++ b/Engine/Source/Programs/UnsyncUI/UnsyncUI/UnsyncUI.csproj
@@ -8,7 +8,7 @@
 		<ApplicationManifest>app.manifest</ApplicationManifest>
         <IncludeAllContentForSelfExtract>true</IncludeAllContentForSelfExtract>
         <ApplicationIcon>unsync_48px.ico</ApplicationIcon>
-        <Version>$(VersionPrefix)1.0.15</Version>
+        <Version>$(VersionPrefix)1.0.16</Version>
 	</PropertyGroup>
 
 </Project>
