UnsyncUI - Ignore inaccessible directories in EnumerateDirectoriesAsync

* Update version to 1.0.17

diff --git a/Engine/Source/Programs/UnsyncUI/UnsyncUI/UnsyncUI.csproj b/Engine/Source/Programs/UnsyncUI/UnsyncUI/UnsyncUI.csproj
--- a/Engine/Source/Programs/UnsyncUI/UnsyncUI/UnsyncUI.csproj
+++ b/Engine/Source/Programs/UnsyncUI/UnsyncUI/UnsyncUI.csproj
@@ -8,7 +8,7 @@
 		<ApplicationManifest>app.manifest</ApplicationManifest>
         <IncludeAllContentForSelfExtract>true</IncludeAllContentForSelfExtract>
         <ApplicationIcon>unsync_48px.ico</ApplicationIcon>
-        <Version>$(VersionPrefix)1.0.16</Version>
+        <Version>$(VersionPrefix)1.0.17</Version>
 	</PropertyGroup>
 
 </Project>
diff --git a/Engine/Source/Programs/UnsyncUI/UnsyncUI/Utils.cs b/Engine/Source/Programs/UnsyncUI/UnsyncUI/Utils.cs
--- a/Engine/Source/Programs/UnsyncUI/UnsyncUI/Utils.cs
+++ b/Engine/Source/Programs/UnsyncUI/UnsyncUI/Utils.cs
@@ -78,6 +78,9 @@ namespace UnsyncUI
 			var tcs = new TaskCompletionSource<IEnumerable<string>>();
             Task.Run(() =>
 			{
+				var enumOptions = new EnumerationOptions();
+				enumOptions.IgnoreInaccessible = true;
+
 				using var cancel = token.Register(() => tcs.TrySetCanceled());
 				try
 				{
@@ -87,7 +90,7 @@ namespace UnsyncUI
 						var timer = new Stopwatch();
 						timer.Start();
 
-						var dirs = Directory.EnumerateDirectories(path).ToList();
+						var dirs = Directory.EnumerateDirectories(path, "*", enumOptions).ToList();
 
 						timer.Stop();
 						Debug.WriteLine($"Time: {timer.Elapsed.TotalSeconds:0.000} s - {path}");
