UnsyncUI - Automatically use scavenge mode when using "Append Build Name" option

diff --git a/Engine/Source/Programs/UnsyncUI/UnsyncUI/BuildSelectorControl.xaml.cs b/Engine/Source/Programs/UnsyncUI/UnsyncUI/BuildSelectorControl.xaml.cs
--- a/Engine/Source/Programs/UnsyncUI/UnsyncUI/BuildSelectorControl.xaml.cs
+++ b/Engine/Source/Programs/UnsyncUI/UnsyncUI/BuildSelectorControl.xaml.cs
@@ -210,7 +210,7 @@ namespace UnsyncUI
 	public sealed class ProjectModel : TabModel
 	{
 		private bool firstRefresh = true;
-		private Action<IEnumerable<(string DstPath, string[] Exclusions, BuildPlatformModel Model)>> onBuildsSelected;
+		private Action<IEnumerable<(SyncStartConfig Config, BuildPlatformModel Model)>> onBuildsSelected;
 
 		private string statusString = null;
 		public string StatusString
@@ -347,7 +347,7 @@ namespace UnsyncUI
 			}
 		}
 
-		public ProjectModel(Config.Project definition, Action<IEnumerable<(string DstPath, string[] Exclusions, BuildPlatformModel Model)>> onBuildsSelected)
+		public ProjectModel(Config.Project definition, Action<IEnumerable<(SyncStartConfig Config, BuildPlatformModel Model)>> onBuildsSelected)
 		{
 			this.Definition = definition;
 			this.onBuildsSelected = onBuildsSelected;
@@ -473,7 +473,16 @@ namespace UnsyncUI
 
 		public void StartSync(IEnumerable<BuildPlatformModel> selectedBuilds)
 		{
-			onBuildsSelected?.Invoke(selectedBuilds.Select(s => (Path.Combine(finalDstPath, s.DestPathRelative), Definition.Exclusions?.ToArray(), s)).ToList());
+			onBuildsSelected?.Invoke(selectedBuilds.Select(s => {
+				var Config = new SyncStartConfig();
+				Config.DstPath = Path.Combine(finalDstPath, s.DestPathRelative);
+				Config.Exclusions = Definition.Exclusions?.ToArray();
+				if (App.Current.UserConfig.AppendBuildName)
+				{
+					Config.ScavengePath = DstPath;
+				}
+				return (Config, s);
+			}).ToList());
 		}
 
 		public override void OnSelected()
diff --git a/Engine/Source/Programs/UnsyncUI/UnsyncUI/JobControl.xaml.cs b/Engine/Source/Programs/UnsyncUI/UnsyncUI/JobControl.xaml.cs
--- a/Engine/Source/Programs/UnsyncUI/UnsyncUI/JobControl.xaml.cs
+++ b/Engine/Source/Programs/UnsyncUI/UnsyncUI/JobControl.xaml.cs
@@ -64,6 +64,7 @@ namespace UnsyncUI
     {
         public BuildPlatformModel Build { get; }
         public string DstPath { get; }
+		public string ScavengePath { get; }
 		public string DstPathBase => Path.GetDirectoryName(DstPath);
         public bool DryRun { get; }
         public string Proxy { get; }
@@ -169,6 +170,7 @@ namespace UnsyncUI
 		public JobModel(
 			BuildPlatformModel build,
 			string dstPath,
+			string scavengePath,
 			bool dryRun,
 			string proxy,
 			string dfs,
@@ -180,6 +182,7 @@ namespace UnsyncUI
 		{
             Build = build;
             DstPath = dstPath;
+			ScavengePath = scavengePath;
             DryRun = dryRun;
             Proxy = proxy;
 			DFS = dfs;
@@ -384,6 +387,11 @@ namespace UnsyncUI
 					args.Add($"--exclude {string.Join(",", Exclusions)}");
 				}
 
+				if (ScavengePath != null)
+				{
+					args.Add($"--scavenge \"{ScavengePath}\"");
+				}
+
 				args.Add($"\"{Build.FullPath}\"");
 				args.Add($"\"{DstPath}\"");
 
diff --git a/Engine/Source/Programs/UnsyncUI/UnsyncUI/MainWindow.xaml.cs b/Engine/Source/Programs/UnsyncUI/UnsyncUI/MainWindow.xaml.cs
--- a/Engine/Source/Programs/UnsyncUI/UnsyncUI/MainWindow.xaml.cs
+++ b/Engine/Source/Programs/UnsyncUI/UnsyncUI/MainWindow.xaml.cs
@@ -63,7 +63,7 @@ namespace UnsyncUI
 
 	public sealed class CustomModel : TabModel
 	{
-		private Action<IEnumerable<(string DstPath, string[] Exclusions, BuildPlatformModel Model)>> onBuildsSelected;
+		private Action<IEnumerable<(SyncStartConfig Config, BuildPlatformModel Model)>> onBuildsSelected;
 
 		public string Name => "Custom";
 
@@ -109,14 +109,17 @@ namespace UnsyncUI
 		private void UpdateSyncCommand()
 			=> OnSyncClicked.Enabled = !string.IsNullOrWhiteSpace(SrcPath) && !string.IsNullOrWhiteSpace(DstPath);
 
-		public CustomModel(Action<IEnumerable<(string DstPath, string[] Exclusions, BuildPlatformModel Model)>> onBuildsSelected)
+		public CustomModel(Action<IEnumerable<(SyncStartConfig Config, BuildPlatformModel Model)>> onBuildsSelected)
 		{
 			this.onBuildsSelected = onBuildsSelected;
 			OnSyncClicked = new Command(() =>
 			{
+				var config = new SyncStartConfig();
+				config.DstPath = DstPath;
+				config.Exclusions = default(string[]);
 				onBuildsSelected(new[]
 				{
-					(DstPath, default(string[]), new BuildPlatformModel(null, null, SrcPath, null, Include))
+					(config, new BuildPlatformModel(null, null, SrcPath, null, Include))
 				});
 			});
 
@@ -124,6 +127,13 @@ namespace UnsyncUI
 		}
 	}
 
+	public class SyncStartConfig
+	{
+		public string DstPath;
+		public string ScavengePath;
+		public string[] Exclusions;
+	};
+
 	public sealed class MainWindowModel : BaseModel
 	{
 		private Config.Proxy selectedProxy = 
@@ -250,18 +260,19 @@ namespace UnsyncUI
 			UpdateProgressState();
 		}
 
-		public void OnBuildsSelected(IEnumerable<(string DstPath, string[] Exclusions, BuildPlatformModel Model)> selectedBuilds)
+		public void OnBuildsSelected(IEnumerable<(SyncStartConfig Config, BuildPlatformModel Model)> selectedBuilds)
 		{
 			foreach (var build in selectedBuilds)
 			{
 				AddJob(new JobModel(
 					build.Model, 
-					build.DstPath, 
+					build.Config.DstPath,
+					build.Config.ScavengePath,
 					DryRun, 
 					SelectedProxy?.Path,
 					DFS, 
 					AdditionalArgs, 
-					build.Exclusions, 
+					build.Config.Exclusions, 
 					OnJobCompleted, 
 					OnClearJob, 
 					OnJobProgress
