UnsyncUI - Trim trailing backslashes from paths when generating unsync command line

diff --git a/Engine/Source/Programs/UnsyncUI/UnsyncUI/JobControl.xaml.cs b/Engine/Source/Programs/UnsyncUI/UnsyncUI/JobControl.xaml.cs
--- a/Engine/Source/Programs/UnsyncUI/UnsyncUI/JobControl.xaml.cs
+++ b/Engine/Source/Programs/UnsyncUI/UnsyncUI/JobControl.xaml.cs
@@ -181,11 +181,11 @@ namespace UnsyncUI
 			Action<JobModel> onProgress)
 		{
             Build = build;
-            DstPath = dstPath;
-			ScavengePath = scavengePath;
+            DstPath = dstPath?.TrimEnd('\\');
+			ScavengePath = scavengePath?.TrimEnd('\\');
             DryRun = dryRun;
-            Proxy = proxy;
-			DFS = dfs;
+            Proxy = proxy?.TrimEnd('\\');
+			DFS = dfs?.TrimEnd('\\');
 			AdditionalArgs = additionalArgs;
 			IncludeFilter = build.Include?.Split(',', StringSplitOptions.RemoveEmptyEntries)
 				.Select(Entry => Entry.Trim())
diff --git a/Engine/Source/Programs/UnsyncUI/UnsyncUI/UnsyncUI.csproj b/Engine/Source/Programs/UnsyncUI/UnsyncUI/UnsyncUI.csproj
--- a/Engine/Source/Programs/UnsyncUI/UnsyncUI/UnsyncUI.csproj
+++ b/Engine/Source/Programs/UnsyncUI/UnsyncUI/UnsyncUI.csproj
@@ -8,7 +8,7 @@
 		<ApplicationManifest>app.manifest</ApplicationManifest>
         <IncludeAllContentForSelfExtract>true</IncludeAllContentForSelfExtract>
         <ApplicationIcon>unsync_48px.ico</ApplicationIcon>
-        <Version>$(VersionPrefix)1.0.17</Version>
+        <Version>$(VersionPrefix)1.0.18</Version>
 	</PropertyGroup>
 
 </Project>
