UnsyncUI - Add experimental support for enumerating builds using proxy server

diff --git a/Engine/Source/Programs/UnsyncUI/UnsyncUI/App.xaml.cs b/Engine/Source/Programs/UnsyncUI/UnsyncUI/App.xaml.cs
--- a/Engine/Source/Programs/UnsyncUI/UnsyncUI/App.xaml.cs
+++ b/Engine/Source/Programs/UnsyncUI/UnsyncUI/App.xaml.cs
@@ -20,6 +20,8 @@ namespace UnsyncUI
 		public UserPreferences UserConfig { get; private set; }
 		public string DefaultSearchTerms { get; private set; } = "";
 
+		internal bool EnableExperimentalFeatures = false;
+
 		protected override void OnStartup(StartupEventArgs e)
 		{
 			string configFile = e.Args.Length > 0 ? e.Args[0] : "unsyncui.xml";
@@ -33,6 +35,10 @@ namespace UnsyncUI
 					DefaultSearchTerms = e.Args[i + 1];
 					++i;
 				}
+				else if (arg == "--experimental")
+				{
+					EnableExperimentalFeatures = true;
+				}
 			}
 
 			var oldWorkingDir = Environment.CurrentDirectory;
@@ -93,6 +99,12 @@ namespace UnsyncUI
 				UserConfig.DFS = null;
 			}
 
+			if (Config != null)
+			{
+				Config.UnsyncPath = UnsyncPath;
+				Config.EnableExperimentalFeatures = EnableExperimentalFeatures;
+			}
+
 			base.OnStartup(e);
 		}
 
diff --git a/Engine/Source/Programs/UnsyncUI/UnsyncUI/BuildSelectorControl.xaml.cs b/Engine/Source/Programs/UnsyncUI/UnsyncUI/BuildSelectorControl.xaml.cs
--- a/Engine/Source/Programs/UnsyncUI/UnsyncUI/BuildSelectorControl.xaml.cs
+++ b/Engine/Source/Programs/UnsyncUI/UnsyncUI/BuildSelectorControl.xaml.cs
@@ -114,7 +114,6 @@ namespace UnsyncUI
 	public sealed class BuildModel : BaseModel
 	{
 		private CancellationTokenSource cts;
-
 		public string Stream => Template.Stream;
 		public string CL => Template.CL;
 		public string Suffix => Template.Suffix;
@@ -122,15 +121,18 @@ namespace UnsyncUI
 		public Config.Directory RootDir { get; }
 		public string Path { get; }
 
-		public Config.BuildTemplate Template { get; }
+		internal Config.BuildTemplate Template { get; }
 
 		public Command OnRefreshPlatformsClicked { get; }
 		public Command OnStopRefreshPlatformsClicked { get; }
 
 		public ObservableCollection<BuildPlatformModel> Platforms { get; } = new ObservableCollection<BuildPlatformModel>();
 
-		public BuildModel(string path, Config.Directory rootDir, Config.BuildTemplate template)
+		private Config AppConfig;
+
+		internal BuildModel(string path, Config.Directory rootDir, Config.BuildTemplate template, Config config)
 		{
+			AppConfig = config;
 			Path = path;
 			RootDir = rootDir;
 			Template = template;
@@ -144,19 +146,19 @@ namespace UnsyncUI
 
 		public bool HasPlatform(string platform) => Platforms.Where(p => p.Platform == platform).Any();
 
-		private async Task EnumeratePlatforms(Config.Directory currentDir, string path, Config.BuildTemplate template, CancellationTokenSource cancellationToken)
+		private async Task EnumeratePlatforms(IDirectoryEnumerator dirEnum, Config.Directory currentDir, string path, Config.BuildTemplate template, CancellationTokenSource cancellationToken)
 		{
 			if (!currentDir.Parse(path, ref template))
 				return;
 
-			foreach (var childDir in await AsyncIO.EnumerateDirectoriesAsync(path, cancellationToken.Token))
+			foreach (var childDir in await dirEnum.EnumerateDirectories(path, cancellationToken.Token))
 			{
 				if (System.IO.Path.GetFileName(childDir) == ".unsync")
 				{
 					List<Config.BuildTemplate> fileGroups = null;
 					if (currentDir.FileGroups?.Any() == true)
 					{
-						var files = await AsyncIO.EnumerateFilesAsync(path, cancellationToken.Token);
+						var files = await dirEnum.EnumerateFiles(path, cancellationToken.Token);
 						currentDir.ParseFileGroups(files.ToList(), template, out fileGroups);
 					}
 					if (fileGroups?.Any() == true)
@@ -177,7 +179,7 @@ namespace UnsyncUI
 				{
 					foreach (var subDir in currentDir.SubDirectories)
 					{
-						await EnumeratePlatforms(subDir, childDir, template, cancellationToken);
+						await EnumeratePlatforms(dirEnum, subDir, childDir, template, cancellationToken);
 					}
 				}
 			}
@@ -193,9 +195,11 @@ namespace UnsyncUI
 
 			Platforms.Clear();
 
+			IDirectoryEnumerator dirEnum = AppConfig.CreateDirectoryEnumerator(Path, RootDir);
+
 			try
 			{
-				await EnumeratePlatforms(RootDir, Path, Template, cts);
+				await EnumeratePlatforms(dirEnum, RootDir, Path, Template, cts);
 			}
 			catch (OperationCanceledException)
 			{ }
diff --git a/Engine/Source/Programs/UnsyncUI/UnsyncUI/Config.cs b/Engine/Source/Programs/UnsyncUI/UnsyncUI/Config.cs
--- a/Engine/Source/Programs/UnsyncUI/UnsyncUI/Config.cs
+++ b/Engine/Source/Programs/UnsyncUI/UnsyncUI/Config.cs
@@ -17,7 +17,7 @@ namespace UnsyncUI
 	public sealed class Config
 	{
 		// Structure that represents mirror server entry in the list from /api/v1/mirrors endpoint
-		class JsonMirrorDesc
+		private class JsonMirrorDesc
 		{
 			public String name { get; set; }
 			public String address { get; set; }
@@ -32,7 +32,7 @@ namespace UnsyncUI
 			public string Path { get; set; }
 		}
 
-		public struct BuildTemplate
+		internal struct BuildTemplate
 		{
 			public string Stream;
 			public string CL;
@@ -82,7 +82,7 @@ namespace UnsyncUI
 				FileGroups = node.Elements("files").Select(f => new FileGroup(f)).ToList();
 			}
 
-			public bool Parse(string path, ref BuildTemplate template)
+			internal bool Parse(string path, ref BuildTemplate template)
 			{
 				var match = Regex.Match(Path.GetFileName(path));
 				if (!match.Success)
@@ -116,7 +116,7 @@ namespace UnsyncUI
 				return true;
 			}
 
-			public void ParseFileGroups(List<String> filePaths, BuildTemplate template, out List<BuildTemplate> groupTemplates)
+			internal void ParseFileGroups(List<String> filePaths, BuildTemplate template, out List<BuildTemplate> groupTemplates)
 			{
 				groupTemplates = null;
 				if (!filePaths.Any())
@@ -154,13 +154,15 @@ namespace UnsyncUI
 
 		public sealed class Project
 		{
+			internal Config AppConfig;
+
 			public string Name { get; set; }
 			public string Root { get; set; }
 			public string Destination { get; set; }
 			public List<string> Exclusions { get; set; }
 			public List<Directory> Children { get; set; }
 
-			private Task EnumerateBuilds(ITargetBlock<BuildModel> pipe, Directory d, string path, BuildTemplate template, CancellationToken cancelToken)
+			private Task EnumerateBuilds(IDirectoryEnumerator dirEnum, ITargetBlock<BuildModel> pipe, Directory d, string path, BuildTemplate template, CancellationToken cancelToken)
 			{
 				return Task.Run(async () =>
 				{
@@ -169,14 +171,14 @@ namespace UnsyncUI
 
 					if (template.IsStreamClFound)
 					{
-						pipe.Post(new BuildModel(path, d, template));
+						pipe.Post(new BuildModel(path, d, template, AppConfig));
 					}
 					else
 					{
 						var jobs = new List<Task>();
-						foreach (var childDir in await AsyncIO.EnumerateDirectoriesAsync(path, cancelToken))
+						foreach (var childDir in await dirEnum.EnumerateDirectories(path, cancelToken))
 						{
-							jobs.AddRange(d.SubDirectories.Select(s => EnumerateBuilds(pipe, s, childDir, template, cancelToken)));
+							jobs.AddRange(d.SubDirectories.Select(s => EnumerateBuilds(dirEnum, pipe, s, childDir, template, cancelToken)));
 						}
 
 						await Task.WhenAll(jobs);
@@ -186,15 +188,17 @@ namespace UnsyncUI
 
 			public (Task, ISourceBlock<BuildModel>) EnumerateBuilds(CancellationToken cancelToken)
 			{
+				IDirectoryEnumerator dirEnum = AppConfig.CreateDirectoryEnumerator(this);
+
 				var pipe = new BufferBlock<BuildModel>();
 				var task = Task.Run(async () =>
 				{
 					try
 					{
 						var allTasks = new List<Task>();
-						foreach (var dir in await AsyncIO.EnumerateDirectoriesAsync(Root, cancelToken))
+						foreach (var dir in await dirEnum.EnumerateDirectories(Root, cancelToken))
 						{
-							allTasks.AddRange(Children.Select(c => EnumerateBuilds(pipe, c, dir, default(BuildTemplate), cancelToken)));
+							allTasks.AddRange(Children.Select(c => EnumerateBuilds(dirEnum, pipe, c, dir, default(BuildTemplate), cancelToken)));
 						}
 
 						await Task.WhenAll(allTasks);
@@ -210,11 +214,14 @@ namespace UnsyncUI
 		}
 
 		public List<Proxy> Proxies { get; set; } = new List<Proxy>();
+		public Proxy RootProxy { get; } = null;
 		public List<Project> Projects { get; set; }
 
 		public string UnsyncPath { get; set; }
 		public string DFS { get; set; }
 
+		public bool EnableExperimentalFeatures { get; set; } = false;
+
 		public Config(string filename)
 		{
 			var rootNode = XDocument.Load(filename).Root;
@@ -243,7 +250,9 @@ namespace UnsyncUI
 			}));
 
 			// Auto-discover proxies
-			List<Proxy> DiscoveredProxies = DiscoverProxies(ConfigProxies);
+			(List<Proxy> DiscoveredProxies, Proxy DiscoveredRootProxy) = DiscoverProxies(ConfigProxies);
+
+			RootProxy = DiscoveredRootProxy;
 
 			if (DiscoveredProxies == null)
 			{
@@ -256,6 +265,7 @@ namespace UnsyncUI
 
 			Projects = rootNode.Element("projects").Elements("project").Select(p => new Project()
 			{
+				AppConfig = this,
 				Name = p.Attribute("name")?.Value,
 				Root = p.Attribute("root")?.Value,
 				Destination = p.Attribute("dest")?.Value,
@@ -264,7 +274,8 @@ namespace UnsyncUI
 			}).ToList();
 		}
 
-		private List<Proxy> DiscoverProxies(List<Proxy> SeedServers)
+		// Returns list of mirrors and the seed proxy server that was used to get it or null
+		private (List<Proxy>, Proxy) DiscoverProxies(List<Proxy> SeedServers)
 		{
 			int DefaultPort = 53841;
 
@@ -329,7 +340,7 @@ namespace UnsyncUI
 
 						if (ParsedProxies.Count != 0)
 						{
-							return ParsedProxies;
+							return (ParsedProxies, SeedServer);
 						}
 					}
 				}
@@ -339,7 +350,52 @@ namespace UnsyncUI
 				}
 			}
 
-			return null;
+			return (null, null);
+		}
+
+		private UnsyncQueryConfig CreateUnsyncQueryConfig()
+		{
+			UnsyncQueryConfig unsyncConfig = new UnsyncQueryConfig();
+			unsyncConfig.proxyAddress = RootProxy.Path;
+			unsyncConfig.unsyncPath = UnsyncPath;
+			return unsyncConfig;
+		}
+
+		private bool CanUseUnsyncDirectoryEnumerator()
+		{
+			return EnableExperimentalFeatures
+				&& RootProxy != null
+				&& UnsyncPath != null;
+		}
+
+		public IDirectoryEnumerator CreateDirectoryEnumerator(Config.Project ProjectSchema)
+		{
+			if (CanUseUnsyncDirectoryEnumerator())
+			{
+				return new UnsyncDirectoryEnumerator(ProjectSchema, CreateUnsyncQueryConfig());
+			}
+			else
+			{
+				return new NativeDirectoryEnumerator();
+			}
+		}
+
+		public IDirectoryEnumerator CreateDirectoryEnumerator(String Path, Config.Directory DirectorySchema)
+		{
+			if (CanUseUnsyncDirectoryEnumerator())
+			{
+				return new UnsyncDirectoryEnumerator(Path, DirectorySchema, CreateUnsyncQueryConfig());
+			}
+			else
+			{
+				return new NativeDirectoryEnumerator();
+			}
 		}
-	}
+
+		public BuildModel CreateBuildModel(string path, Config.Directory rootDir)
+		{
+			return new BuildModel(path, rootDir, default(BuildTemplate), this);
+		}
+
+	} // class Config
 }
diff --git a/Engine/Source/Programs/UnsyncUI/UnsyncUI/JobControl.xaml.cs b/Engine/Source/Programs/UnsyncUI/UnsyncUI/JobControl.xaml.cs
--- a/Engine/Source/Programs/UnsyncUI/UnsyncUI/JobControl.xaml.cs
+++ b/Engine/Source/Programs/UnsyncUI/UnsyncUI/JobControl.xaml.cs
@@ -371,6 +371,11 @@ namespace UnsyncUI
 				if (!string.IsNullOrWhiteSpace(Proxy))
 				{
 					args.Add($"--proxy \"{Proxy}\"");
+
+					if (App.Current.Config.EnableExperimentalFeatures)
+					{
+						args.Add("--login");
+					}
 				}
 
 				if (!string.IsNullOrWhiteSpace(DFS))
@@ -405,7 +410,7 @@ namespace UnsyncUI
 
 				Debug.WriteLine($"Running unsync with args {argsStr}");
 
-				var unsyncPath = App.Current.UnsyncPath;
+				var unsyncPath = App.Current.Config.UnsyncPath;
 				proc = new AsyncProcess(unsyncPath, argsStr);
 
 				await foreach (var str in proc.RunAsync(cts.Token))
diff --git a/Engine/Source/Programs/UnsyncUI/UnsyncUI/MainWindow.xaml.cs b/Engine/Source/Programs/UnsyncUI/UnsyncUI/MainWindow.xaml.cs
--- a/Engine/Source/Programs/UnsyncUI/UnsyncUI/MainWindow.xaml.cs
+++ b/Engine/Source/Programs/UnsyncUI/UnsyncUI/MainWindow.xaml.cs
@@ -23,6 +23,11 @@ namespace UnsyncUI
 			DataContext = new MainWindowModel();
 			var version = System.Reflection.Assembly.GetEntryAssembly().GetName().Version;
 			Title = $"UnsyncUI v{version.Major}.{version.Minor}.{version.Build}";
+
+			if (App.Current.EnableExperimentalFeatures)
+			{
+				Title += " [EXPERIMENTAL]";
+			}
 		}
 
 		private void Window_Closing(object sender, CancelEventArgs e)
diff --git a/Engine/Source/Programs/UnsyncUI/UnsyncUI/Query.cs b/Engine/Source/Programs/UnsyncUI/UnsyncUI/Query.cs
new file mode 100644
--- /dev/null
+++ b/Engine/Source/Programs/UnsyncUI/UnsyncUI/Query.cs
@@ -0,0 +1,313 @@
+// Copyright Epic Games, Inc. All Rights Reserved.
+
+using System.Collections.Generic;
+using System.Threading.Tasks;
+using System.Threading;
+using System;
+using System.IO;
+using System.Diagnostics;
+using System.Text.Json;
+using System.Linq;
+
+namespace UnsyncUI
+{
+	public class UnsyncQueryConfig
+	{
+		public String unsyncPath;
+		public String proxyAddress;
+	}
+
+	class SearchQueryResultEntry
+	{
+		public String path { get; set; }
+
+		public bool is_directory { get; set; }
+		public UInt64 size { get; set; }
+		public UInt64 mtime { get; set; }
+	}
+
+	class SearchQueryResult
+	{
+		public String root { get; set; }
+		public List<SearchQueryResultEntry> entries { get; set; }
+	}
+
+	public class UnsyncDirectoryEnumerator : IDirectoryEnumerator
+	{
+		private UnsyncQueryConfig Config;
+
+		private Config.Project ProjectSchema;
+
+		private String DirectoryRoot;
+		private Config.Directory DirectorySchema;
+
+		private bool Initialized = false;
+
+		private class Entry
+		{
+			public String FullPath;
+			public int Depth;
+			public UInt64 Size;
+			public UInt64 MTime;
+			public bool IsDirectory;
+		}
+
+		List<Entry> Entries;
+
+		public UnsyncDirectoryEnumerator(Config.Project ProjectSchema, UnsyncQueryConfig Config)
+		{
+			this.Config = Config;
+			this.ProjectSchema = ProjectSchema;
+		}
+		public UnsyncDirectoryEnumerator(String Path, Config.Directory DirectorySchema, UnsyncQueryConfig Config)
+		{
+			this.Config = Config;
+			this.DirectoryRoot = Path;
+			this.DirectorySchema = DirectorySchema;
+		}
+
+		private async Task LazyInit(CancellationToken cancellationToken)
+		{
+			if (!Initialized)
+			{
+				if (ProjectSchema != null)
+				{
+					await InitProject(cancellationToken);
+				}
+				else if (DirectorySchema != null)
+				{
+					await InitDirectory(cancellationToken);
+				}
+				else
+				{
+					throw new Exception("Unexpected UnsyncDirectoryEnumerator configuration. Either project or directory must be specified.");
+				}
+			}
+		}
+		class BuildTraversalState
+		{
+			public bool FoundCL = false;
+			public bool FoundStream = false;
+		}
+
+		private void ProcessQueryResult(SearchQueryResult queryResult)
+		{
+			if (Entries == null)
+			{
+				Entries = new List<Entry>();
+			}
+
+			Entries.Clear();
+
+			foreach (var queryEntry in queryResult.entries)
+			{
+				Entry entry = new Entry();
+				entry.FullPath = Path.Combine(queryResult.root, queryEntry.path);
+				entry.Depth = GetDirectoryDepth(entry.FullPath);
+				entry.IsDirectory = queryEntry.is_directory;
+				entry.MTime = queryEntry.mtime;
+				entry.Size = queryEntry.size;
+				Entries.Add(entry);
+			}
+		}
+
+		private async Task RunQueries(List<String> queryStrings, CancellationToken cancellationToken)
+		{
+			foreach (String query in queryStrings)
+			{
+				String argsStr = query + $" --proxy {Config.proxyAddress}";
+
+				Debug.WriteLine($"Running: unsync.exe {argsStr}");
+
+				var proc = new AsyncProcess(Config.unsyncPath, argsStr);
+				var responseJson = "";
+				await foreach (var str in proc.RunAsync(cancellationToken))
+				{
+					responseJson += str;
+				}
+
+				if (proc.ExitCode == 0)
+				{
+					try
+					{
+						SearchQueryResult queryResult = JsonSerializer.Deserialize<SearchQueryResult>(responseJson);
+						ProcessQueryResult(queryResult);
+					}
+					catch (Exception ex)
+					{
+						Debug.WriteLine("Exception while parsing unsync query JSON: " + ex.Message);
+					}
+				}
+			}
+		}
+
+		private async Task InitProject(CancellationToken cancellationToken)
+		{
+			String baseQuery = $"query search \"{ProjectSchema.Root}\"";
+
+			BuildTraversalState traversalState = new BuildTraversalState();
+			List<String> queryStrings = new List<String>();
+
+			foreach (Config.Directory childDir in ProjectSchema.Children)
+			{
+				GenerateQueries(ref traversalState, childDir, baseQuery, ref queryStrings);
+			}
+
+			await RunQueries(queryStrings, cancellationToken);
+
+			Initialized = true;
+		}
+
+		private void GenerateQueries(ref BuildTraversalState traversalState, Config.Directory dir, String query, ref List<String> output)
+		{
+			if (traversalState != null)
+			{
+				if (dir.CL != null)
+				{
+					traversalState.FoundCL = true;
+				}
+
+				if (dir.Stream != null)
+				{
+					traversalState.FoundStream = true;
+				}
+
+				if (traversalState.FoundCL && traversalState.FoundStream)
+				{
+					// Stop visiting directories early when rules for matching CL and Stream are found in the project config.
+					// Typically this just includes the project root directory, but if the builds are organized
+					// by Stream and then by CL, we want to traverse a deeper hierarchy.
+
+					output.Add(query);
+
+					return;
+				}
+			}
+
+			// Adding a regex rule for the directory will instruct unsync to list all of its children
+			query += $" \"{dir.Regex}\"";
+
+			if (dir.SubDirectories.Any())
+			{
+				foreach (Config.Directory childDir in dir.SubDirectories)
+				{
+					GenerateQueries(ref traversalState, childDir, query, ref output);
+				}
+			}
+			else
+			{
+				output.Add(query);
+			}
+		}
+
+		private async Task InitDirectory(CancellationToken cancellationToken)
+		{
+			String baseQuery = $"query search \"{DirectoryRoot}\"";
+
+			BuildTraversalState traversalState = null;
+			List<String> queryStrings = new List<String>();
+
+			foreach (Config.Directory childDir in DirectorySchema.SubDirectories)
+			{
+				GenerateQueries(ref traversalState, childDir, baseQuery, ref queryStrings);
+			}
+
+			await RunQueries(queryStrings, cancellationToken);
+
+			Initialized = true;
+		}
+
+		private int GetDirectoryDepth(string path)
+		{
+			int depth = 0;
+			foreach (var c in path)
+			{
+				if (c == Path.DirectorySeparatorChar)
+				{
+					depth += 1;
+				}
+			}
+			return depth;
+		}
+
+		private struct EntryFilter
+		{
+			public bool IncludeFiles;
+			public bool IncludeDirectories;
+
+			public bool ShouldInclude(Entry entry)
+			{
+				if (entry.IsDirectory && !IncludeDirectories)
+				{
+					return false;
+				}
+
+				if (!entry.IsDirectory && !IncludeFiles)
+				{
+					return false;
+				}
+
+				return true;
+			}
+		}
+
+		private Task<IEnumerable<string>> EnumerateEntries(string path, CancellationToken token, EntryFilter filter)
+		{
+			var tcs = new TaskCompletionSource<IEnumerable<string>>();
+			Task.Run(async () =>
+			{
+				using var cancel = token.Register(() => tcs.TrySetCanceled());
+				try
+				{
+					await LazyInit(token);
+					var dirs = new List<string>();
+
+					int pathDepth = GetDirectoryDepth(path);
+
+					string requiredPrefix = path + Path.DirectorySeparatorChar;
+
+					// Doesn't run very often and typically matches most entries anyway,
+					// so can just do a naive search without building any hierarchies.
+					foreach (var entry in Entries)
+					{
+						if (!filter.ShouldInclude(entry))
+						{
+							continue;
+						}
+
+						// We want only directories that are immediate children of the given path
+						if (entry.Depth == 1 + pathDepth
+							&& entry.FullPath.StartsWith(requiredPrefix))
+						{
+							//String relativePath = Path.GetRelativePath(path, entry.FullPath);
+							dirs.Add(entry.FullPath);
+						}
+					}
+
+					tcs.TrySetResult(dirs);
+				}
+				catch (OperationCanceledException)
+				{
+					tcs.TrySetCanceled();
+				}
+			});
+			return tcs.Task;
+		}
+
+		public Task<IEnumerable<string>> EnumerateDirectories(string path, CancellationToken token)
+		{
+			EntryFilter filter;
+			filter.IncludeDirectories = true;
+			filter.IncludeFiles = false;
+			return EnumerateEntries(path, token, filter);
+		}
+		public Task<IEnumerable<string>> EnumerateFiles(string path, CancellationToken token)
+		{
+			EntryFilter filter;
+			filter.IncludeDirectories = false;
+			filter.IncludeFiles = true;
+			return EnumerateEntries(path, token, filter);
+		}
+	}
+
+}
diff --git a/Engine/Source/Programs/UnsyncUI/UnsyncUI/Utils.cs b/Engine/Source/Programs/UnsyncUI/UnsyncUI/Utils.cs
--- a/Engine/Source/Programs/UnsyncUI/UnsyncUI/Utils.cs
+++ b/Engine/Source/Programs/UnsyncUI/UnsyncUI/Utils.cs
@@ -69,7 +69,25 @@ namespace UnsyncUI
         { }
     }
 
-    public static class AsyncIO
+	public interface IDirectoryEnumerator
+	{
+		public Task<IEnumerable<string>> EnumerateDirectories(string path, CancellationToken token);
+		public Task<IEnumerable<string>> EnumerateFiles(string path, CancellationToken token);
+	}
+
+	public class NativeDirectoryEnumerator : IDirectoryEnumerator
+	{
+		public Task<IEnumerable<string>> EnumerateDirectories(string path, CancellationToken token)
+		{
+			return AsyncIO.EnumerateDirectoriesAsync(path, token);
+		}
+		public Task<IEnumerable<string>> EnumerateFiles(string path, CancellationToken token)
+		{
+			return AsyncIO.EnumerateFilesAsync(path, token);
+		}
+	}
+
+	static class AsyncIO
     {
 		private static SemaphoreSlim mutex = new SemaphoreSlim(4);
 
