UnsyncUI - Change .net framework from 3.0-core to 6.0-windows (same as UGS)

diff --git a/Engine/Source/Programs/UnsyncUI/UnsyncUI/UnsyncUI.csproj b/Engine/Source/Programs/UnsyncUI/UnsyncUI/UnsyncUI.csproj
--- a/Engine/Source/Programs/UnsyncUI/UnsyncUI/UnsyncUI.csproj
+++ b/Engine/Source/Programs/UnsyncUI/UnsyncUI/UnsyncUI.csproj
@@ -2,7 +2,7 @@
 
 	<PropertyGroup>
 		<OutputType>WinExe</OutputType>
-		<TargetFramework>netcoreapp3.1</TargetFramework>
+		<TargetFramework>net6.0-windows</TargetFramework>
 		<UseWPF>true</UseWPF>
 		<UseWindowsForms>true</UseWindowsForms>
 		<ApplicationManifest>app.manifest</ApplicationManifest>
