UnsyncUI - Suppress compiler warning about WebRequest being obsolete

diff --git a/Engine/Source/Programs/UnsyncUI/UnsyncUI/Config.cs b/Engine/Source/Programs/UnsyncUI/UnsyncUI/Config.cs
--- a/Engine/Source/Programs/UnsyncUI/UnsyncUI/Config.cs
+++ b/Engine/Source/Programs/UnsyncUI/UnsyncUI/Config.cs
@@ -12,6 +12,9 @@ using System.Xml.Linq;
 using System.Net;
 using System.Text.Json;
 
+// warning SYSLIB0014: 'WebRequest.Create(string)' is obsolete: 'WebRequest, HttpWebRequest, ServicePoint, and WebClient are obsolete. Use HttpClient instead.' (https://aka.ms/dotnet-warnings/SYSLIB0014)
+#pragma warning disable SYSLIB0014
+
 namespace UnsyncUI
 {
 	public sealed class Config
