UnsyncUI - Remove DFS text field, move Dry Run to the bottom, add option to ignore StdErr when running async process

diff --git a/Engine/Source/Programs/UnsyncUI/UnsyncUI/App.xaml.cs b/Engine/Source/Programs/UnsyncUI/UnsyncUI/App.xaml.cs
--- a/Engine/Source/Programs/UnsyncUI/UnsyncUI/App.xaml.cs
+++ b/Engine/Source/Programs/UnsyncUI/UnsyncUI/App.xaml.cs
@@ -94,10 +94,6 @@ namespace UnsyncUI
 			}
 
 			UserConfig = UserPreferences.Load();
-			if (UserConfig.DFS == "")
-			{
-				UserConfig.DFS = null;
-			}
 
 			if (Config != null)
 			{
@@ -123,7 +119,6 @@ namespace UnsyncUI
 		private static readonly string userFile = Path.Combine(userDir, $"user_{version}.json");
 
 		public string Proxy { get; set; }
-		public string DFS { get; set; }
 		public Dictionary<string, string> ProjectDestinationMap { get; set; } = new Dictionary<string, string>();
 		public string CustomSrcPath { get; set; }
 		public string CustomDstPath { get; set; }
diff --git a/Engine/Source/Programs/UnsyncUI/UnsyncUI/BuildSelectorControl.xaml b/Engine/Source/Programs/UnsyncUI/UnsyncUI/BuildSelectorControl.xaml
--- a/Engine/Source/Programs/UnsyncUI/UnsyncUI/BuildSelectorControl.xaml
+++ b/Engine/Source/Programs/UnsyncUI/UnsyncUI/BuildSelectorControl.xaml
@@ -178,16 +178,19 @@
 		</GroupBox>
 
 		<!-- Start sync -->
-		<TextBlock Grid.Row="3" Margin="4" HorizontalAlignment="Center">
-			<Hyperlink Command="{Binding OnStartSync}" CommandParameter="{Binding SelectedItems, ElementName=SelectedPlatforms}">
-				<Hyperlink.IsEnabled>
-					<MultiBinding Converter="{StaticResource BoolConv_All}">
-						<Binding ElementName="SelectedPlatforms" Path="SelectedItems.Count"/>
-						<Binding Path="DstPath.Length"/>
-					</MultiBinding>
-				</Hyperlink.IsEnabled>
-				Start Sync
-			</Hyperlink>
-		</TextBlock>
+		<Grid Grid.Row="3">
+			<CheckBox Content="Dry Run" IsChecked="{Binding DryRun}" ToolTip="Download and verify the data, but don't write outputs to disk" Margin="10,4,4,4" HorizontalAlignment="Left"/>
+			<TextBlock Margin="4" HorizontalAlignment="Center">
+				<Hyperlink Command="{Binding OnStartSync}" CommandParameter="{Binding SelectedItems, ElementName=SelectedPlatforms}">
+					<Hyperlink.IsEnabled>
+						<MultiBinding Converter="{StaticResource BoolConv_All}">
+							<Binding ElementName="SelectedPlatforms" Path="SelectedItems.Count"/>
+							<Binding Path="DstPath.Length"/>
+						</MultiBinding>
+					</Hyperlink.IsEnabled>
+					Start Sync
+				</Hyperlink>
+			</TextBlock>
+		</Grid>
 	</Grid>
 </UserControl>
diff --git a/Engine/Source/Programs/UnsyncUI/UnsyncUI/BuildSelectorControl.xaml.cs b/Engine/Source/Programs/UnsyncUI/UnsyncUI/BuildSelectorControl.xaml.cs
--- a/Engine/Source/Programs/UnsyncUI/UnsyncUI/BuildSelectorControl.xaml.cs
+++ b/Engine/Source/Programs/UnsyncUI/UnsyncUI/BuildSelectorControl.xaml.cs
@@ -337,6 +337,8 @@ namespace UnsyncUI
 			}
 		}
 
+		public bool DryRun { get; set; }
+
 		private bool appendBuildName = true;
 		public bool AppendBuildName
 		{
@@ -481,6 +483,7 @@ namespace UnsyncUI
 				var Config = new SyncStartConfig();
 				Config.DstPath = Path.Combine(finalDstPath, s.DestPathRelative);
 				Config.Exclusions = Definition.Exclusions?.ToArray();
+				Config.DryRun = DryRun;
 				if (App.Current.UserConfig.AppendBuildName)
 				{
 					Config.ScavengePath = DstPath;
diff --git a/Engine/Source/Programs/UnsyncUI/UnsyncUI/Config.cs b/Engine/Source/Programs/UnsyncUI/UnsyncUI/Config.cs
--- a/Engine/Source/Programs/UnsyncUI/UnsyncUI/Config.cs
+++ b/Engine/Source/Programs/UnsyncUI/UnsyncUI/Config.cs
@@ -221,7 +221,6 @@ namespace UnsyncUI
 		public List<Project> Projects { get; set; }
 
 		public string UnsyncPath { get; set; }
-		public string DFS { get; set; }
 
 		public bool EnableExperimentalFeatures { get; set; } = false;
 
@@ -236,8 +235,6 @@ namespace UnsyncUI
 				throw new Exception("Unable to find unsync.exe binary specified in config file.");
 			}
 
-			DFS = rootNode.Attribute("dfs")?.Value;
-
 			Proxies.Add(new Proxy()
 			{
 				Name = "(none)",
diff --git a/Engine/Source/Programs/UnsyncUI/UnsyncUI/JobControl.xaml.cs b/Engine/Source/Programs/UnsyncUI/UnsyncUI/JobControl.xaml.cs
--- a/Engine/Source/Programs/UnsyncUI/UnsyncUI/JobControl.xaml.cs
+++ b/Engine/Source/Programs/UnsyncUI/UnsyncUI/JobControl.xaml.cs
@@ -68,7 +68,6 @@ namespace UnsyncUI
 		public string DstPathBase => Path.GetDirectoryName(DstPath);
         public bool DryRun { get; }
         public string Proxy { get; }
-		public string DFS { get; }
 		public string AdditionalArgs { get; }
 		public string[] IncludeFilter { get; }
 		public string[] Exclusions { get; }
@@ -173,7 +172,6 @@ namespace UnsyncUI
 			string scavengePath,
 			bool dryRun,
 			string proxy,
-			string dfs,
 			string additionalArgs,
 			string[] exclusions,
 			Action<JobModel> onCompletion,
@@ -185,7 +183,6 @@ namespace UnsyncUI
 			ScavengePath = scavengePath?.TrimEnd('\\');
             DryRun = dryRun;
             Proxy = proxy?.TrimEnd('\\');
-			DFS = dfs?.TrimEnd('\\');
 			AdditionalArgs = additionalArgs;
 			IncludeFilter = build.Include?.Split(',', StringSplitOptions.RemoveEmptyEntries)
 				.Select(Entry => Entry.Trim())
@@ -378,11 +375,6 @@ namespace UnsyncUI
 					}
 				}
 
-				if (!string.IsNullOrWhiteSpace(DFS))
-				{
-					args.Add($"--dfs {DFS}");
-				}
-
 				if (IncludeFilter?.Length > 0)
 				{
 					args.Add($"--include {string.Join(",", IncludeFilter)}");
diff --git a/Engine/Source/Programs/UnsyncUI/UnsyncUI/MainWindow.xaml b/Engine/Source/Programs/UnsyncUI/UnsyncUI/MainWindow.xaml
--- a/Engine/Source/Programs/UnsyncUI/UnsyncUI/MainWindow.xaml
+++ b/Engine/Source/Programs/UnsyncUI/UnsyncUI/MainWindow.xaml
@@ -57,9 +57,12 @@
 							<TextBox Text="{Binding Include}" ToolTip="Comma-separated list parts of filenames to filter the download"/>
 						</GroupBox>
 
-						<TextBlock Grid.Row="3" Margin="4" HorizontalAlignment="Center">
-							<Hyperlink Command="{Binding OnSyncClicked}">Start Custom Sync</Hyperlink>
-						</TextBlock>
+						<Grid Grid.Row="3" >
+							<CheckBox Content="Dry Run" IsChecked="{Binding DryRun}" ToolTip="Download and verify the data, but don't write outputs to disk" Margin="10,4,4,4" HorizontalAlignment="Left"/>
+							<TextBlock Margin="4" HorizontalAlignment="Center">
+								<Hyperlink Command="{Binding OnSyncClicked}">Start Custom Sync</Hyperlink>
+							</TextBlock>
+						</Grid>
 
 					</Grid>
 				</DataTemplate>
@@ -97,25 +100,9 @@
 			</ComboBox>
 		</GroupBox>
 
-		<GroupBox Grid.Row="1" Header="Advanced">
+		<GroupBox Grid.Row="1" Header="Additional arguments">
 			<Grid>
-				<Grid.ColumnDefinitions>
-					<ColumnDefinition Width="auto"/>
-					<ColumnDefinition Width="*"/>
-					<ColumnDefinition Width="auto"/>
-				</Grid.ColumnDefinitions>
-
-				<Grid.RowDefinitions>
-					<RowDefinition/>
-					<RowDefinition/>
-				</Grid.RowDefinitions>
-
-				<CheckBox Grid.Column="0" VerticalAlignment="Center" IsChecked="{Binding DryRun}" Content="Dry Run" HorizontalAlignment="Center" Margin="32,4,32,4"/>
-				<TextBox Grid.Column="1" Text="{Binding DFS}"/>
-				<TextBlock Grid.Column="2" VerticalAlignment="Center" Text="DFS Filter" Margin="16,4,16,4"/>
-
-				<TextBox   Grid.Row="1" Grid.Column="0" Grid.ColumnSpan="2" Text="{Binding AdditionalArgs}" />
-				<TextBlock Grid.Row="1" Grid.Column="2" VerticalAlignment="Center" Text="Additional Args" Margin="16,4,16,4"/>
+				<TextBox Grid.Row="1" Grid.Column="0" Grid.ColumnSpan="2" Text="{Binding AdditionalArgs}" />
 			</Grid>
 		</GroupBox>
 
diff --git a/Engine/Source/Programs/UnsyncUI/UnsyncUI/MainWindow.xaml.cs b/Engine/Source/Programs/UnsyncUI/UnsyncUI/MainWindow.xaml.cs
--- a/Engine/Source/Programs/UnsyncUI/UnsyncUI/MainWindow.xaml.cs
+++ b/Engine/Source/Programs/UnsyncUI/UnsyncUI/MainWindow.xaml.cs
@@ -85,6 +85,8 @@ namespace UnsyncUI
 			}
 		}
 
+		public bool DryRun { get; set; }
+
 		private string dstPath = App.Current.UserConfig.CustomDstPath;
 		public string DstPath
 		{
@@ -120,6 +122,7 @@ namespace UnsyncUI
 			OnSyncClicked = new Command(() =>
 			{
 				var config = new SyncStartConfig();
+				config.DryRun = DryRun;
 				config.DstPath = DstPath;
 				config.Exclusions = default(string[]);
 				onBuildsSelected(new[]
@@ -134,6 +137,7 @@ namespace UnsyncUI
 
 	public class SyncStartConfig
 	{
+		public bool DryRun;
 		public string DstPath;
 		public string ScavengePath;
 		public string[] Exclusions;
@@ -155,24 +159,6 @@ namespace UnsyncUI
 			}
 		}
 
-		private bool dryRun = false;
-		public bool DryRun
-		{
-			get => dryRun;
-			set => SetProperty(ref dryRun, value);
-		}
-
-		private string dfs = App.Current.UserConfig.DFS ?? App.Current.Config?.DFS;
-		public string DFS
-		{
-			get => dfs;
-			set
-			{
-				SetProperty(ref dfs, value);
-				App.Current.UserConfig.DFS = value;
-			}
-		}
-
 		private bool showHelp = true;
 		public bool ShowHelp
 		{
@@ -273,9 +259,8 @@ namespace UnsyncUI
 					build.Model, 
 					build.Config.DstPath,
 					build.Config.ScavengePath,
-					DryRun, 
+					build.Config.DryRun, 
 					SelectedProxy?.Path,
-					DFS, 
 					AdditionalArgs, 
 					build.Config.Exclusions, 
 					OnJobCompleted, 
diff --git a/Engine/Source/Programs/UnsyncUI/UnsyncUI/Query.cs b/Engine/Source/Programs/UnsyncUI/UnsyncUI/Query.cs
--- a/Engine/Source/Programs/UnsyncUI/UnsyncUI/Query.cs
+++ b/Engine/Source/Programs/UnsyncUI/UnsyncUI/Query.cs
@@ -119,7 +119,8 @@ namespace UnsyncUI
 
 				var proc = new AsyncProcess(Config.unsyncPath, argsStr);
 				var responseJson = "";
-				await foreach (var str in proc.RunAsync(cancellationToken))
+				// TODO: read stderr stream and somehow report status/errors
+				await foreach (var str in proc.RunAsync(cancellationToken, false /*ReadStdErr*/))
 				{
 					responseJson += str;
 				}
diff --git a/Engine/Source/Programs/UnsyncUI/UnsyncUI/Utils.cs b/Engine/Source/Programs/UnsyncUI/UnsyncUI/Utils.cs
--- a/Engine/Source/Programs/UnsyncUI/UnsyncUI/Utils.cs
+++ b/Engine/Source/Programs/UnsyncUI/UnsyncUI/Utils.cs
@@ -286,7 +286,7 @@ namespace UnsyncUI
 			};
 		}
 
-		public async IAsyncEnumerable<string> RunAsync([EnumeratorCancellation] CancellationToken cancelToken)
+		public async IAsyncEnumerable<string> RunAsync([EnumeratorCancellation] CancellationToken cancelToken, bool ReadStdErr = true)
 		{
 			try
 			{
@@ -303,7 +303,7 @@ namespace UnsyncUI
 
 					var pipe = new BufferBlock<string>();
 
-					async Task ReadStream(Stream stream)
+					async Task ReadStream(Stream stream, bool ShouldPost)
 					{
 						var block = new byte[4096];
 						var mem = new Memory<byte>(block);
@@ -314,15 +314,19 @@ namespace UnsyncUI
 							if (bytesRead == 0)
 								break;
 
-							// @todo: this won't handle UTF-8 encoding if a char is split across the read boundary.
-							pipe.Post(Encoding.UTF8.GetString(mem.Span.Slice(0, bytesRead)));
+							if (ShouldPost)
+							{
+								// @todo: this won't handle UTF-8 encoding if a char is split across the read boundary.
+								pipe.Post(Encoding.UTF8.GetString(mem.Span.Slice(0, bytesRead)));
+							}
+							
 						}
 
 						pipe.Complete();
 					}
 
-					var stdoutTask = ReadStream(proc.StandardOutput.BaseStream);
-					var stderrTask = ReadStream(proc.StandardError.BaseStream);
+					var stdoutTask = ReadStream(proc.StandardOutput.BaseStream, true);
+					var stderrTask = ReadStream(proc.StandardError.BaseStream, ReadStdErr);
 
 					var completionTask = Task.Run(async () =>
 					{
