UnsyncUI - Add explicit login / logout buttons to the UI (only in experimental mode)

diff --git a/Engine/Source/Programs/UnsyncUI/UnsyncUI/App.xaml.cs b/Engine/Source/Programs/UnsyncUI/UnsyncUI/App.xaml.cs
--- a/Engine/Source/Programs/UnsyncUI/UnsyncUI/App.xaml.cs
+++ b/Engine/Source/Programs/UnsyncUI/UnsyncUI/App.xaml.cs
@@ -125,6 +125,7 @@ namespace UnsyncUI
 		public string CustomInclude { get; set; }
 		public string AdditionalArgs { get; set; }
 		public bool AppendBuildName { get; set; } = false;
+		public bool LogInOnStartup { get; set; } = false;
 
 		public static UserPreferences Load()
 		{
diff --git a/Engine/Source/Programs/UnsyncUI/UnsyncUI/Config.cs b/Engine/Source/Programs/UnsyncUI/UnsyncUI/Config.cs
--- a/Engine/Source/Programs/UnsyncUI/UnsyncUI/Config.cs
+++ b/Engine/Source/Programs/UnsyncUI/UnsyncUI/Config.cs
@@ -224,6 +224,8 @@ namespace UnsyncUI
 
 		public bool EnableExperimentalFeatures { get; set; } = false;
 
+		internal string loggedInUser;
+
 		public Config(string filename)
 		{
 			var rootNode = XDocument.Load(filename).Root;
@@ -365,7 +367,8 @@ namespace UnsyncUI
 		{
 			return EnableExperimentalFeatures
 				&& RootProxy != null
-				&& UnsyncPath != null;
+				&& UnsyncPath != null
+				&& loggedInUser != null;
 		}
 
 		public IDirectoryEnumerator CreateDirectoryEnumerator(Config.Project ProjectSchema)
diff --git a/Engine/Source/Programs/UnsyncUI/UnsyncUI/MainWindow.xaml b/Engine/Source/Programs/UnsyncUI/UnsyncUI/MainWindow.xaml
--- a/Engine/Source/Programs/UnsyncUI/UnsyncUI/MainWindow.xaml
+++ b/Engine/Source/Programs/UnsyncUI/UnsyncUI/MainWindow.xaml
@@ -84,20 +84,39 @@
 		</Grid.ColumnDefinitions>
 
 		<GroupBox Grid.Row="0" Header="Proxy">
-			<ComboBox  ItemsSource="{Binding Config.Proxies}" SelectedItem="{Binding SelectedProxy}">
-				<ComboBox.ItemTemplate>
-					<DataTemplate>
-						<TextBlock>
-								<Run Text="{Binding Name}" />
-								<TextBlock Visibility="{Binding Path, Converter={StaticResource VisConv_Collapsed}}">
-									<Run Text=" ("/>
-									<Run Text="{Binding Path}"/>
-									<Run Text=")"/>
-								</TextBlock>
-						</TextBlock>
-					</DataTemplate>
-				</ComboBox.ItemTemplate>
-			</ComboBox>
+			<Grid>
+				<Grid.RowDefinitions>
+					<RowDefinition Height="auto" />
+					<RowDefinition Height="auto" />
+				</Grid.RowDefinitions>
+				<ComboBox Grid.Row="0" ItemsSource="{Binding Config.Proxies}" SelectedItem="{Binding SelectedProxy}">
+					<ComboBox.ItemTemplate>
+						<DataTemplate>
+							<TextBlock>
+									<Run Text="{Binding Name}" />
+									<TextBlock Visibility="{Binding Path, Converter={StaticResource VisConv_Collapsed}}">
+										<Run Text=" ("/>
+										<Run Text="{Binding Path}"/>
+										<Run Text=")"/>
+									</TextBlock>
+							</TextBlock>
+						</DataTemplate>
+					</ComboBox.ItemTemplate>
+				</ComboBox>
+				<TextBlock Grid.Row="1" Visibility="{Binding Path=ShouldShowLoginInfo, Converter={StaticResource VisConv_Collapsed}}">
+					<Run Text="User:"/>
+					<!--Logged in-->
+					<TextBlock Visibility="{Binding Path=LoggedInUser, Converter={StaticResource VisConv_Collapsed}}">
+						<Run FontWeight="Bold" Text="{Binding Path=LoggedInUser, Mode=OneWay}"/>
+						<Run Text=" - "/>
+						<Hyperlink Command="{Binding OnLogOutClicked}">Log out</Hyperlink>
+					</TextBlock>
+					<!--Not logged in-->
+					<TextBlock Visibility="{Binding Path=LoggedInUser, Converter={StaticResource VisConv_Collapsed_Inv}}">
+						<Hyperlink FontWeight="Bold" Command="{Binding OnLogInClicked}">Log in</Hyperlink>
+					</TextBlock>
+				</TextBlock>
+			</Grid>
 		</GroupBox>
 
 		<GroupBox Grid.Row="1" Header="Additional arguments">
diff --git a/Engine/Source/Programs/UnsyncUI/UnsyncUI/MainWindow.xaml.cs b/Engine/Source/Programs/UnsyncUI/UnsyncUI/MainWindow.xaml.cs
--- a/Engine/Source/Programs/UnsyncUI/UnsyncUI/MainWindow.xaml.cs
+++ b/Engine/Source/Programs/UnsyncUI/UnsyncUI/MainWindow.xaml.cs
@@ -177,6 +177,50 @@ namespace UnsyncUI
 			}
 		}
 
+		public bool ShouldShowLoginInfo { get => App.Current.EnableExperimentalFeatures; }
+
+		public string LoggedInUser
+		{
+			get => Config.loggedInUser;
+			set 
+			{
+				App.Current.UserConfig.LogInOnStartup = (value != null);
+				SetProperty(ref Config.loggedInUser, value);
+			}
+		}
+
+		public Command OnLogInClicked { get; }
+		public Command OnLogOutClicked { get; }
+
+		private void LogIn()
+		{
+			if (Config.RootProxy.Path == null
+				|| Config.UnsyncPath == null)
+			{
+				LoggedInUser = null;
+				return;
+			}
+
+			UnsyncQueryUtil queryUtil = new UnsyncQueryUtil(Config.UnsyncPath, Config.RootProxy.Path);
+
+			try
+			{
+				LoginQueryResult LoginInfo = queryUtil.Login();
+				LoggedInUser = LoginInfo.sub;
+			}
+			catch (Exception ex)
+			{
+				// TODO: add global status/log window
+				Debug.WriteLine($"Login failed with exception: {ex}");
+			}
+		}
+
+		private void LogOut()
+		{
+			LoggedInUser = null;
+		}
+
+
 		private TaskbarItemProgressState progressState = TaskbarItemProgressState.None;
 		public TaskbarItemProgressState ProgressState
 		{
@@ -226,6 +270,9 @@ namespace UnsyncUI
 			OnClearQueueClicked = new Command(ClearQueue) { Enabled = true };
 			OnClearCompletedClicked = new Command(ClearCompleted) { Enabled = true };
 
+			OnLogInClicked = new Command(LogIn) { Enabled = true };
+			OnLogOutClicked = new Command(LogOut) { Enabled = true };
+
 			if (Config != null)
 			{
 				foreach (var p in Config.Projects)
@@ -234,6 +281,12 @@ namespace UnsyncUI
 				}
 			}
 
+			if (Config.EnableExperimentalFeatures
+				&& App.Current.UserConfig.LogInOnStartup)
+			{
+				LogIn();
+			}
+
 			Tabs.Add(new CustomModel(OnBuildsSelected));
 
 			SelectedTab = Tabs.FirstOrDefault();
diff --git a/Engine/Source/Programs/UnsyncUI/UnsyncUI/Query.cs b/Engine/Source/Programs/UnsyncUI/UnsyncUI/Query.cs
--- a/Engine/Source/Programs/UnsyncUI/UnsyncUI/Query.cs
+++ b/Engine/Source/Programs/UnsyncUI/UnsyncUI/Query.cs
@@ -8,6 +8,7 @@ using System.IO;
 using System.Diagnostics;
 using System.Text.Json;
 using System.Linq;
+using System.Collections;
 
 namespace UnsyncUI
 {
@@ -32,6 +33,58 @@ namespace UnsyncUI
 		public List<SearchQueryResultEntry> entries { get; set; }
 	}
 
+	// Subset of decoded JWT claims
+	public class LoginQueryResult
+	{
+		public String aud { get; set; }
+		public String sub { get; set; }
+		public String iss { get; set; }
+		public List<String> scp { get; set; }
+		public UInt64 iat { get; set; }
+		public UInt64 exp { get; set; }
+		public List<String> groups { get; set; }
+	}
+
+	public class UnsyncQueryUtil
+	{
+		UnsyncQueryConfig Config;
+		public UnsyncQueryUtil(UnsyncQueryConfig InConfig)
+		{
+			Config = InConfig;
+		}
+
+		public UnsyncQueryUtil(string unsyncPath, string proxyAddress)
+		{
+			Config = new UnsyncQueryConfig();
+			Config.unsyncPath = unsyncPath;
+			Config.proxyAddress = proxyAddress;
+		}
+
+		public LoginQueryResult Login()
+		{
+			String argsStr = $"login --decode --proxy {Config.proxyAddress}";
+			AsyncProcess proc = new AsyncProcess(Config.unsyncPath, argsStr);
+			CancellationToken cancellationToken = new CancellationToken();
+
+			var responseJson = "";
+
+			var LoginTask = Task.Run(async () => {
+				// TODO: read stderr stream and somehow report status/errors
+				await foreach (var str in proc.RunAsync(cancellationToken, false /*ReadStdErr*/))
+				{
+					responseJson += str;
+				}
+			});
+			LoginTask.Wait();
+
+			LoginQueryResult queryResult = JsonSerializer.Deserialize<LoginQueryResult>(responseJson);
+
+			return queryResult;
+		}
+	}
+
+	
+
 	public class UnsyncDirectoryEnumerator : IDirectoryEnumerator
 	{
 		private UnsyncQueryConfig Config;
diff --git a/Engine/Source/Programs/UnsyncUI/UnsyncUI/UnsyncUI.csproj b/Engine/Source/Programs/UnsyncUI/UnsyncUI/UnsyncUI.csproj
--- a/Engine/Source/Programs/UnsyncUI/UnsyncUI/UnsyncUI.csproj
+++ b/Engine/Source/Programs/UnsyncUI/UnsyncUI/UnsyncUI.csproj
@@ -8,7 +8,7 @@
 		<ApplicationManifest>app.manifest</ApplicationManifest>
         <IncludeAllContentForSelfExtract>true</IncludeAllContentForSelfExtract>
         <ApplicationIcon>unsync_48px.ico</ApplicationIcon>
-        <Version>$(VersionPrefix)1.0.19</Version>
+        <Version>$(VersionPrefix)1.0.20</Version>
 	</PropertyGroup>
 
 </Project>
diff --git a/Engine/Source/Programs/UnsyncUI/UnsyncUI/Utils.cs b/Engine/Source/Programs/UnsyncUI/UnsyncUI/Utils.cs
--- a/Engine/Source/Programs/UnsyncUI/UnsyncUI/Utils.cs
+++ b/Engine/Source/Programs/UnsyncUI/UnsyncUI/Utils.cs
@@ -350,5 +350,16 @@ namespace UnsyncUI
 				cancelToken.ThrowIfCancellationRequested();
 			}
 		}
+
+		public async Task<string> RunAndGetOutput(bool ReadStdErr = true)
+		{
+			var result = "";
+			var cancellationToken = new CancellationToken();
+			await foreach (var str in RunAsync(cancellationToken, false /*ReadStdErr*/))
+			{
+				result += str;
+			}
+			return result;
+		}
 	}
 }
