UnsyncUI - Add application log panel, enable user authentication UI by default

* Fix race condition in process stdout/err stream handling which could result in some pending data to be lost at the end of the process
* Bump version to 1.0.21

diff --git a/Engine/Source/Programs/UnsyncUI/UnsyncUI/App.xaml.cs b/Engine/Source/Programs/UnsyncUI/UnsyncUI/App.xaml.cs
--- a/Engine/Source/Programs/UnsyncUI/UnsyncUI/App.xaml.cs
+++ b/Engine/Source/Programs/UnsyncUI/UnsyncUI/App.xaml.cs
@@ -2,9 +2,12 @@
 
 using System;
 using System.Collections.Generic;
+using System.Diagnostics;
 using System.IO;
 using System.Text.Json;
+using System.Threading;
 using System.Windows;
+using System.Xml.Linq;
 
 namespace UnsyncUI
 {
@@ -21,6 +24,34 @@ namespace UnsyncUI
 		public string DefaultSearchTerms { get; private set; } = "";
 
 		internal bool EnableExperimentalFeatures = false;
+		internal bool EnableUserAuthentication = true;
+
+		internal string ApplicationLog { get; private set; } = "";
+
+		internal void LogError(string message)
+		{
+			LogMessage("ERROR: " + message);
+		}
+
+		internal void LogMessage(string message)
+		{
+			LogDebug(message);
+
+			Dispatcher.InvokeAsync(delegate
+			{
+				ApplicationLog += message + "\n";
+				var model = MainWindow.DataContext as MainWindowModel;
+				if (model != null)
+				{
+					model.OnLogUpdated();
+				}
+			});
+		}
+
+		internal void LogDebug(string message)
+		{
+			Debug.WriteLine(message);
+		}
 
 		protected override void OnStartup(StartupEventArgs e)
 		{
@@ -99,6 +130,7 @@ namespace UnsyncUI
 			{
 				Config.UnsyncPath = UnsyncPath;
 				Config.EnableExperimentalFeatures = EnableExperimentalFeatures;
+				Config.EnableUserAuthentication = EnableUserAuthentication;
 			}
 
 			base.OnStartup(e);
diff --git a/Engine/Source/Programs/UnsyncUI/UnsyncUI/BuildSelectorControl.xaml.cs b/Engine/Source/Programs/UnsyncUI/UnsyncUI/BuildSelectorControl.xaml.cs
--- a/Engine/Source/Programs/UnsyncUI/UnsyncUI/BuildSelectorControl.xaml.cs
+++ b/Engine/Source/Programs/UnsyncUI/UnsyncUI/BuildSelectorControl.xaml.cs
@@ -460,14 +460,14 @@ namespace UnsyncUI
 				}
 				else
 				{
-					Debug.WriteLine($"Search took {timer.Elapsed.TotalSeconds} seconds.");
+					App.Current.LogMessage($"Search took {timer.Elapsed.TotalSeconds} seconds.");
 				}
 			}
 			catch (OperationCanceledException)
 			{ }
 			catch (Exception ex)
 			{
-				Debug.WriteLine($"Directory enumeration failed with exception: {ex}");
+				App.Current.LogMessage($"Directory enumeration failed with exception: {ex}");
 				StatusString = ex.Message;
 			}
 			finally
diff --git a/Engine/Source/Programs/UnsyncUI/UnsyncUI/Config.cs b/Engine/Source/Programs/UnsyncUI/UnsyncUI/Config.cs
--- a/Engine/Source/Programs/UnsyncUI/UnsyncUI/Config.cs
+++ b/Engine/Source/Programs/UnsyncUI/UnsyncUI/Config.cs
@@ -224,6 +224,8 @@ namespace UnsyncUI
 
 		public bool EnableExperimentalFeatures { get; set; } = false;
 
+		public bool EnableUserAuthentication { get; set; } = true;
+
 		internal string loggedInUser;
 
 		public Config(string filename)
@@ -365,7 +367,7 @@ namespace UnsyncUI
 
 		private bool CanUseUnsyncDirectoryEnumerator()
 		{
-			return EnableExperimentalFeatures
+			return EnableUserAuthentication
 				&& RootProxy != null
 				&& UnsyncPath != null
 				&& loggedInUser != null;
diff --git a/Engine/Source/Programs/UnsyncUI/UnsyncUI/JobControl.xaml.cs b/Engine/Source/Programs/UnsyncUI/UnsyncUI/JobControl.xaml.cs
--- a/Engine/Source/Programs/UnsyncUI/UnsyncUI/JobControl.xaml.cs
+++ b/Engine/Source/Programs/UnsyncUI/UnsyncUI/JobControl.xaml.cs
@@ -369,7 +369,8 @@ namespace UnsyncUI
 				{
 					args.Add($"--proxy \"{Proxy}\"");
 
-					if (App.Current.Config.EnableExperimentalFeatures)
+					if (App.Current.Config.EnableUserAuthentication
+						&& App.Current.Config.loggedInUser != null)
 					{
 						args.Add("--login");
 					}
@@ -400,7 +401,7 @@ namespace UnsyncUI
 					argsStr = AdditionalArgs.Substring(1);
 				}
 
-				Debug.WriteLine($"Running unsync with args {argsStr}");
+				App.Current.LogMessage($"Running unsync with args {argsStr}");
 
 				var unsyncPath = App.Current.Config.UnsyncPath;
 				proc = new AsyncProcess(unsyncPath, argsStr);
diff --git a/Engine/Source/Programs/UnsyncUI/UnsyncUI/MainWindow.xaml b/Engine/Source/Programs/UnsyncUI/UnsyncUI/MainWindow.xaml
--- a/Engine/Source/Programs/UnsyncUI/UnsyncUI/MainWindow.xaml
+++ b/Engine/Source/Programs/UnsyncUI/UnsyncUI/MainWindow.xaml
@@ -72,27 +72,29 @@
 
 	</Window.Resources>
 	<Grid Margin="4">
-		<Grid.RowDefinitions>
-			<RowDefinition Height="auto" />
-			<RowDefinition Height="auto" />
-			<RowDefinition Height="*" />
-		</Grid.RowDefinitions>
 		<Grid.ColumnDefinitions>
 			<ColumnDefinition Width="476"/>
 			<ColumnDefinition Width="8"/>
 			<ColumnDefinition Width="*"/>
 		</Grid.ColumnDefinitions>
 
-		<GroupBox Grid.Row="0" Header="Proxy">
-			<Grid>
-				<Grid.RowDefinitions>
-					<RowDefinition Height="auto" />
-					<RowDefinition Height="auto" />
-				</Grid.RowDefinitions>
-				<ComboBox Grid.Row="0" ItemsSource="{Binding Config.Proxies}" SelectedItem="{Binding SelectedProxy}">
-					<ComboBox.ItemTemplate>
-						<DataTemplate>
-							<TextBlock>
+		<!--Left panel-->
+		<Grid Grid.Row="0">
+			<Grid.RowDefinitions>
+				<RowDefinition Height="auto" />
+				<RowDefinition Height="auto" />
+				<RowDefinition Height="*" />
+			</Grid.RowDefinitions>
+			<GroupBox Grid.Row="0" Header="Proxy">
+				<Grid>
+					<Grid.RowDefinitions>
+						<RowDefinition Height="auto" />
+						<RowDefinition Height="auto" />
+					</Grid.RowDefinitions>
+					<ComboBox Grid.Row="0" ItemsSource="{Binding Config.Proxies}" SelectedItem="{Binding SelectedProxy}">
+						<ComboBox.ItemTemplate>
+							<DataTemplate>
+								<TextBlock>
 									<Run Text="{Binding Name}" />
 									<TextBlock Visibility="{Binding Path, Converter={StaticResource VisConv_Collapsed}}">
 										<Run Text=" ("/>
@@ -100,10 +102,10 @@
 										<Run Text=")"/>
 									</TextBlock>
 							</TextBlock>
-						</DataTemplate>
-					</ComboBox.ItemTemplate>
-				</ComboBox>
-				<TextBlock Grid.Row="1" Visibility="{Binding Path=ShouldShowLoginInfo, Converter={StaticResource VisConv_Collapsed}}">
+							</DataTemplate>
+						</ComboBox.ItemTemplate>
+					</ComboBox>
+					<TextBlock Grid.Row="1" Visibility="{Binding Path=ShouldShowLoginInfo, Converter={StaticResource VisConv_Collapsed}}">
 					<Run Text="User:"/>
 					<!--Logged in-->
 					<TextBlock Visibility="{Binding Path=LoggedInUser, Converter={StaticResource VisConv_Collapsed}}">
@@ -116,107 +118,146 @@
 						<Hyperlink FontWeight="Bold" Command="{Binding OnLogInClicked}">Log in</Hyperlink>
 					</TextBlock>
 				</TextBlock>
-			</Grid>
-		</GroupBox>
+				</Grid>
+			</GroupBox>
 
-		<GroupBox Grid.Row="1" Header="Additional arguments">
-			<Grid>
-				<TextBox Grid.Row="1" Grid.Column="0" Grid.ColumnSpan="2" Text="{Binding AdditionalArgs}" />
-			</Grid>
-		</GroupBox>
-
-		<GroupBox Grid.Row="2" Header="Projects">
-			<Grid>
-				<Grid.RowDefinitions>
-					<RowDefinition Height="auto"/>
-					<RowDefinition Height="*"/>
-				</Grid.RowDefinitions>
-				<TabControl Grid.Row ="1" ItemsSource="{Binding Tabs}" SelectedIndex="0" ContentTemplateSelector="{Binding Source={StaticResource mainTabSelector}}" SelectedItem="{Binding SelectedTab}">
-					<TabControl.ItemTemplate>
-						<DataTemplate>
-							<TextBlock Text="{Binding Name}"/>
-						</DataTemplate>
-					</TabControl.ItemTemplate>
-				</TabControl>
-			</Grid>
-		</GroupBox>
+			<GroupBox Grid.Row="1" Header="Additional arguments">
+				<Grid>
+					<TextBox Grid.Row="1" Grid.Column="0" Grid.ColumnSpan="2" Text="{Binding AdditionalArgs}" />
+				</Grid>
+			</GroupBox>
+
+			<GroupBox Grid.Row="2" Header="Projects">
+				<Grid>
+					<Grid.RowDefinitions>
+						<RowDefinition Height="auto"/>
+						<RowDefinition Height="*"/>
+					</Grid.RowDefinitions>
+					<TabControl Grid.Row ="1" ItemsSource="{Binding Tabs}" SelectedIndex="0" ContentTemplateSelector="{Binding Source={StaticResource mainTabSelector}}" SelectedItem="{Binding SelectedTab}">
+						<TabControl.ItemTemplate>
+							<DataTemplate>
+								<TextBlock Text="{Binding Name}"/>
+							</DataTemplate>
+						</TabControl.ItemTemplate>
+					</TabControl>
+				</Grid>
+			</GroupBox>
+
+		</Grid>
 
 		<GridSplitter Grid.Column="1" Grid.RowSpan="3" HorizontalAlignment="Stretch" Foreground="Transparent" Background="Transparent" />
 
-		<TextBlock Grid.Row="0" Grid.Column="2" Grid.RowSpan="3" TextAlignment="Center" VerticalAlignment="Center">
-			<TextBlock.Visibility>
-				<MultiBinding Converter="{StaticResource VisConv_All_Collapsed}">
-					<Binding Path="QueuedJobs.Count" Converter="{StaticResource NotConv}" />
-					<Binding Path="ActiveJob" Converter="{StaticResource NotConv}" />
-					<Binding Path="CompletedJobs.Count" Converter="{StaticResource NotConv}" />
-				</MultiBinding>
-			</TextBlock.Visibility>
-			<Run Text="To get started, search for builds in the left pane,"/>
-			<LineBreak/>
-			<Run Text="select one or more platforms and a destination folder, and click Start Sync."/>
-		</TextBlock>
-
-		<Grid Grid.Row="0" Grid.Column="2" Grid.RowSpan="3">
-			<Grid.Resources>
-				<GridLength x:Key="glStar">*</GridLength>
-				<GridLength x:Key="glAuto">auto</GridLength>
-				<local:ObjectConverter x:Key="rdHeightConv" TrueValue="{StaticResource glStar}" FalseValue="{StaticResource glAuto}"/>
-				<local:ObjectConverter x:Key="rdMinHeightConv" TrueValue="60" FalseValue="0"/>
-			</Grid.Resources>
+		<!--Right panel-->
+		<Grid Grid.Row="0" Grid.Column="2" >
 			<Grid.RowDefinitions>
-				<RowDefinition Height="auto"/>
-				<RowDefinition Height="{Binding ActiveJob, Converter={StaticResource rdHeightConv}}" MinHeight="{Binding ActiveJob, Converter={StaticResource rdMinHeightConv}}"/>
-				<RowDefinition Height="auto"/>
+				<RowDefinition Height="*"/>
+				<RowDefinition Height="4"/>
+				<RowDefinition Height="100"/>
 			</Grid.RowDefinitions>
+			<!-- Sync job logs -->
+			<Grid Grid.Row="0">
+				<!-- Help text -->
+				<TextBlock Grid.Row="0" Grid.Column="2" Grid.RowSpan="3" TextAlignment="Center" VerticalAlignment="Center">
+					<TextBlock.Visibility>
+						<MultiBinding Converter="{StaticResource VisConv_All_Collapsed}">
+							<Binding Path="QueuedJobs.Count" Converter="{StaticResource NotConv}" />
+							<Binding Path="ActiveJob" Converter="{StaticResource NotConv}" />
+							<Binding Path="CompletedJobs.Count" Converter="{StaticResource NotConv}" />
+						</MultiBinding>
+					</TextBlock.Visibility>
+					<Run Text="To get started, search for builds in the left pane,"/>
+					<LineBreak/>
+					<Run Text="select one or more platforms and a destination folder, and click Start Sync."/>
+				</TextBlock>
+				<!-- Job logs -->
+				<Grid Grid.Row="0" Grid.Column="2" Grid.RowSpan="3">
+					<Grid.Resources>
+						<GridLength x:Key="glStar">*</GridLength>
+						<GridLength x:Key="glAuto">auto</GridLength>
+						<local:ObjectConverter x:Key="rdHeightConv" TrueValue="{StaticResource glStar}" FalseValue="{StaticResource glAuto}"/>
+						<local:ObjectConverter x:Key="rdMinHeightConv" TrueValue="60" FalseValue="0"/>
+					</Grid.Resources>
+					<Grid.RowDefinitions>
+						<RowDefinition Height="auto"/>
+						<RowDefinition Height="{Binding ActiveJob, Converter={StaticResource rdHeightConv}}" MinHeight="{Binding ActiveJob, Converter={StaticResource rdMinHeightConv}}"/>
+						<RowDefinition Height="auto"/>
+					</Grid.RowDefinitions>
 
-			<!-- Queued Jobs Pane -->
-			<GroupBox Grid.Row="0" Visibility="{Binding QueuedJobs.Count, Converter={StaticResource VisConv_Collapsed}}">
-				<GroupBox.Header>
-					<TextBlock VerticalAlignment="Center" HorizontalAlignment="Center">
-                            <Run Text="{Binding QueuedJobs.Count, Mode=OneWay}" />
-                            <Run Text="Queued Job(s)" />
-                            <TextBlock>
-                                <Run Text=" - " />
-                                <Hyperlink Command="{Binding OnClearQueueClicked}">Cancel All</Hyperlink>
-                            </TextBlock>
-                        </TextBlock>
-				</GroupBox.Header>
-				<ItemsControl ItemsSource="{Binding QueuedJobs}">
-					<ItemsControl.ItemTemplate>
-						<ItemContainerTemplate>
-							<local:JobControl/>
-						</ItemContainerTemplate>
-					</ItemsControl.ItemTemplate>
-				</ItemsControl>
-			</GroupBox>
+					<!-- Queued Jobs Pane -->
+					<GroupBox Grid.Row="0" Visibility="{Binding QueuedJobs.Count, Converter={StaticResource VisConv_Collapsed}}">
+						<GroupBox.Header>
+							<TextBlock VerticalAlignment="Center" HorizontalAlignment="Center">
+									<Run Text="{Binding QueuedJobs.Count, Mode=OneWay}" />
+									<Run Text="Queued Job(s)" />
+									<TextBlock>
+										<Run Text=" - " />
+										<Hyperlink Command="{Binding OnClearQueueClicked}">Cancel All</Hyperlink>
+									</TextBlock>
+								</TextBlock>
+						</GroupBox.Header>
+						<ItemsControl ItemsSource="{Binding QueuedJobs}">
+							<ItemsControl.ItemTemplate>
+								<ItemContainerTemplate>
+									<local:JobControl/>
+								</ItemContainerTemplate>
+							</ItemsControl.ItemTemplate>
+						</ItemsControl>
+					</GroupBox>
 
-			<!-- Active Job Pane -->
-			<GroupBox Grid.Row="1" Visibility="{Binding ActiveJob, Converter={StaticResource VisConv_Collapsed}}" Header="Active Job">
-				<local:JobControl DataContext="{Binding ActiveJob}"/>
-			</GroupBox>
+					<!-- Active Job Pane -->
+					<GroupBox Grid.Row="1" Visibility="{Binding ActiveJob, Converter={StaticResource VisConv_Collapsed}}" Header="Active Job">
+						<local:JobControl DataContext="{Binding ActiveJob}"/>
+					</GroupBox>
 
-			<!-- Completed Jobs Pane -->
-			<GroupBox Grid.Row="2"  Visibility="{Binding CompletedJobs.Count, Converter={StaticResource VisConv_Collapsed}}">
-				<GroupBox.Header>
-					<TextBlock VerticalAlignment="Center" HorizontalAlignment="Center">
-                        <Run Text="{Binding CompletedJobs.Count, Mode=OneWay}" />
-                        <Run Text="Completed Job(s)" />
-                        <TextBlock Visibility="{Binding CompletedJobs.Count, Converter={StaticResource VisConv_Collapsed}}">
-                            <Run Text=" - " />
-                            <Hyperlink Command="{Binding OnClearCompletedClicked}">Clear All</Hyperlink>
-                        </TextBlock>
-                    </TextBlock>
-				</GroupBox.Header>
-
-				<ItemsControl ItemsSource="{Binding CompletedJobs}">
-					<ItemsControl.ItemTemplate>
-						<ItemContainerTemplate>
-							<local:JobControl/>
-						</ItemContainerTemplate>
-					</ItemsControl.ItemTemplate>
-				</ItemsControl>
-			</GroupBox>
+					<!-- Completed Jobs Pane -->
+					<GroupBox Grid.Row="2"  Visibility="{Binding CompletedJobs.Count, Converter={StaticResource VisConv_Collapsed}}">
+						<GroupBox.Header>
+							<TextBlock VerticalAlignment="Center" HorizontalAlignment="Center">
+								<Run Text="{Binding CompletedJobs.Count, Mode=OneWay}" />
+								<Run Text="Completed Job(s)" />
+								<TextBlock Visibility="{Binding CompletedJobs.Count, Converter={StaticResource VisConv_Collapsed}}">
+									<Run Text=" - " />
+									<Hyperlink Command="{Binding OnClearCompletedClicked}">Clear All</Hyperlink>
+								</TextBlock>
+							</TextBlock>
+						</GroupBox.Header>
+
+						<ItemsControl ItemsSource="{Binding CompletedJobs}">
+							<ItemsControl.ItemTemplate>
+								<ItemContainerTemplate>
+									<local:JobControl/>
+								</ItemContainerTemplate>
+							</ItemsControl.ItemTemplate>
+						</ItemsControl>
+					</GroupBox>
+				</Grid>
+			</Grid>
+			<GridSplitter Grid.Row="1" HorizontalAlignment="Stretch" />
+			<!-- Main app log -->
+			<Grid Grid.Row="2">
+				<GroupBox>
+					<GroupBox.Header>
+						<WrapPanel>
+							<!--
+								TODO: expandable log window
+								<Expander VerticalAlignment="Center" Name="ApplicationLogExpander" IsExpanded="{Binding LogExpanded}" />
+							-->
+							<TextBlock VerticalAlignment="Center"><Run Text="Application Log"/></TextBlock>
+						</WrapPanel>
+					</GroupBox.Header>
+					<local:AutoScrollTextBox 
+						Text="{Binding ApplicationLog, Mode=OneWay}" 
+						IsReadOnly="True"
+						IsReadOnlyCaretVisible="True"
+						VerticalScrollBarVisibility="Visible"
+						HorizontalScrollBarVisibility="Auto"
+						TextWrapping="Wrap"
+						>
+					</local:AutoScrollTextBox>
+				</GroupBox>
+			</Grid>
 		</Grid>
+
 	</Grid>
 
 </Window>
diff --git a/Engine/Source/Programs/UnsyncUI/UnsyncUI/MainWindow.xaml.cs b/Engine/Source/Programs/UnsyncUI/UnsyncUI/MainWindow.xaml.cs
--- a/Engine/Source/Programs/UnsyncUI/UnsyncUI/MainWindow.xaml.cs
+++ b/Engine/Source/Programs/UnsyncUI/UnsyncUI/MainWindow.xaml.cs
@@ -50,6 +50,15 @@ namespace UnsyncUI
 		}
 	}
 
+	public class AutoScrollTextBox : TextBox
+	{
+		protected override void OnTextChanged(TextChangedEventArgs e)
+		{
+			base.OnTextChanged(e);
+			ScrollToEnd();
+		}
+	}
+
 	public sealed class MainTabTemplateSelector : DataTemplateSelector
 	{
 		public DataTemplate ProjectTemplate { get; set; }
@@ -159,6 +168,23 @@ namespace UnsyncUI
 			}
 		}
 
+		private bool logExpanded = true;
+		public bool LogExpanded
+		{
+			get => logExpanded;
+			set => SetProperty(ref logExpanded, value);
+		}
+
+		public string ApplicationLog
+		{
+			get => App.Current.ApplicationLog;
+		}
+
+		public void OnLogUpdated()
+		{
+			OnPropertyChanged("ApplicationLog");
+		}
+
 		private bool showHelp = true;
 		public bool ShowHelp
 		{
@@ -177,7 +203,7 @@ namespace UnsyncUI
 			}
 		}
 
-		public bool ShouldShowLoginInfo { get => App.Current.EnableExperimentalFeatures; }
+		public bool ShouldShowLoginInfo { get => App.Current.EnableUserAuthentication; }
 
 		public string LoggedInUser
 		{
@@ -211,7 +237,7 @@ namespace UnsyncUI
 			catch (Exception ex)
 			{
 				// TODO: add global status/log window
-				Debug.WriteLine($"Login failed with exception: {ex}");
+				App.Current.LogError($"Login failed with exception: {ex}");
 			}
 		}
 
@@ -281,7 +307,7 @@ namespace UnsyncUI
 				}
 			}
 
-			if (Config.EnableExperimentalFeatures
+			if (Config.EnableUserAuthentication
 				&& App.Current.UserConfig.LogInOnStartup)
 			{
 				LogIn();
diff --git a/Engine/Source/Programs/UnsyncUI/UnsyncUI/Query.cs b/Engine/Source/Programs/UnsyncUI/UnsyncUI/Query.cs
--- a/Engine/Source/Programs/UnsyncUI/UnsyncUI/Query.cs
+++ b/Engine/Source/Programs/UnsyncUI/UnsyncUI/Query.cs
@@ -168,8 +168,6 @@ namespace UnsyncUI
 			{
 				String argsStr = query + $" --proxy {Config.proxyAddress}";
 
-				Debug.WriteLine($"Running: unsync.exe {argsStr}");
-
 				var proc = new AsyncProcess(Config.unsyncPath, argsStr);
 				var responseJson = "";
 				// TODO: read stderr stream and somehow report status/errors
@@ -187,7 +185,7 @@ namespace UnsyncUI
 					}
 					catch (Exception ex)
 					{
-						Debug.WriteLine("Exception while parsing unsync query JSON: " + ex.Message);
+						App.Current.LogError("Exception while parsing unsync query JSON: " + ex.Message);
 					}
 				}
 			}
diff --git a/Engine/Source/Programs/UnsyncUI/UnsyncUI/UnsyncUI.csproj b/Engine/Source/Programs/UnsyncUI/UnsyncUI/UnsyncUI.csproj
--- a/Engine/Source/Programs/UnsyncUI/UnsyncUI/UnsyncUI.csproj
+++ b/Engine/Source/Programs/UnsyncUI/UnsyncUI/UnsyncUI.csproj
@@ -8,7 +8,7 @@
 		<ApplicationManifest>app.manifest</ApplicationManifest>
         <IncludeAllContentForSelfExtract>true</IncludeAllContentForSelfExtract>
         <ApplicationIcon>unsync_48px.ico</ApplicationIcon>
-        <Version>$(VersionPrefix)1.0.20</Version>
+        <Version>$(VersionPrefix)1.0.21</Version>
 	</PropertyGroup>
 
 </Project>
diff --git a/Engine/Source/Programs/UnsyncUI/UnsyncUI/Utils.cs b/Engine/Source/Programs/UnsyncUI/UnsyncUI/Utils.cs
--- a/Engine/Source/Programs/UnsyncUI/UnsyncUI/Utils.cs
+++ b/Engine/Source/Programs/UnsyncUI/UnsyncUI/Utils.cs
@@ -93,6 +93,8 @@ namespace UnsyncUI
 
         public static Task<IEnumerable<string>> EnumerateDirectoriesAsync(string path, CancellationToken token)
         {
+			App.Current.LogMessage($"Enumerating: {path}");
+
 			var tcs = new TaskCompletionSource<IEnumerable<string>>();
             Task.Run(() =>
 			{
@@ -105,14 +107,7 @@ namespace UnsyncUI
 					mutex.Wait(token);
 					try
 					{
-						var timer = new Stopwatch();
-						timer.Start();
-
 						var dirs = Directory.EnumerateDirectories(path, "*", enumOptions).ToList();
-
-						timer.Stop();
-						Debug.WriteLine($"Time: {timer.Elapsed.TotalSeconds:0.000} s - {path}");
-
 						tcs.TrySetResult(dirs);
 					}
 					catch (Exception ex)
@@ -148,7 +143,7 @@ namespace UnsyncUI
 						var dirs = Directory.EnumerateFiles(path).ToList();
 
 						timer.Stop();
-						Debug.WriteLine($"Time: {timer.Elapsed.TotalSeconds:0.000} s - {path}");
+						App.Current.LogMessage($"Time: {timer.Elapsed.TotalSeconds:0.000} s - {path}");
 
 						tcs.TrySetResult(dirs);
 					}
@@ -288,6 +283,9 @@ namespace UnsyncUI
 
 		public async IAsyncEnumerable<string> RunAsync([EnumeratorCancellation] CancellationToken cancelToken, bool ReadStdErr = true)
 		{
+			string processFileName = Path.GetFileName(proc.StartInfo.FileName);
+			App.Current.LogMessage($"Running: {processFileName} {proc.StartInfo.Arguments}");
+
 			try
 			{
 				using (var cancel = cancelToken.Register(() =>
@@ -312,17 +310,17 @@ namespace UnsyncUI
 						{
 							int bytesRead = await stream.ReadAsync(mem);
 							if (bytesRead == 0)
+							{
 								break;
+							}
 
 							if (ShouldPost)
 							{
 								// @todo: this won't handle UTF-8 encoding if a char is split across the read boundary.
-								pipe.Post(Encoding.UTF8.GetString(mem.Span.Slice(0, bytesRead)));
+								string decodedString = Encoding.UTF8.GetString(mem.Span.Slice(0, bytesRead));
+								pipe.Post(decodedString);
 							}
-							
 						}
-
-						pipe.Complete();
 					}
 
 					var stdoutTask = ReadStream(proc.StandardOutput.BaseStream, true);
@@ -336,7 +334,8 @@ namespace UnsyncUI
 
 					while (await pipe.OutputAvailableAsync())
 					{
-						yield return pipe.Receive();
+						string receivedString = pipe.Receive();
+						yield return receivedString;
 					}
 
 					await completionTask;
@@ -347,6 +346,8 @@ namespace UnsyncUI
 				ExitCode = proc.ExitCode;
 				proc.Dispose();
 
+				App.Current.LogDebug($"Finished: {processFileName} with exit code {ExitCode}");
+
 				cancelToken.ThrowIfCancellationRequested();
 			}
 		}
