UnsyncUI - Mark user login section as experimental, add log clear button

diff --git a/Engine/Source/Programs/UnsyncUI/UnsyncUI/App.xaml.cs b/Engine/Source/Programs/UnsyncUI/UnsyncUI/App.xaml.cs
--- a/Engine/Source/Programs/UnsyncUI/UnsyncUI/App.xaml.cs
+++ b/Engine/Source/Programs/UnsyncUI/UnsyncUI/App.xaml.cs
@@ -28,6 +28,19 @@ namespace UnsyncUI
 
 		internal string ApplicationLog { get; private set; } = "";
 
+		internal void ClearApplicationLog()
+		{
+			Dispatcher.InvokeAsync(delegate
+			{
+				ApplicationLog = "";
+				var model = MainWindow.DataContext as MainWindowModel;
+				if (model != null)
+				{
+					model.OnLogUpdated();
+				}
+			});
+		}
+
 		internal void LogError(string message)
 		{
 			LogMessage("ERROR: " + message);
diff --git a/Engine/Source/Programs/UnsyncUI/UnsyncUI/MainWindow.xaml b/Engine/Source/Programs/UnsyncUI/UnsyncUI/MainWindow.xaml
--- a/Engine/Source/Programs/UnsyncUI/UnsyncUI/MainWindow.xaml
+++ b/Engine/Source/Programs/UnsyncUI/UnsyncUI/MainWindow.xaml
@@ -105,19 +105,25 @@
 							</DataTemplate>
 						</ComboBox.ItemTemplate>
 					</ComboBox>
-					<TextBlock Grid.Row="1" Visibility="{Binding Path=ShouldShowLoginInfo, Converter={StaticResource VisConv_Collapsed}}">
-					<Run Text="User:"/>
-					<!--Logged in-->
-					<TextBlock Visibility="{Binding Path=LoggedInUser, Converter={StaticResource VisConv_Collapsed}}">
-						<Run FontWeight="Bold" Text="{Binding Path=LoggedInUser, Mode=OneWay}"/>
-						<Run Text=" - "/>
-						<Hyperlink Command="{Binding OnLogOutClicked}">Log out</Hyperlink>
-					</TextBlock>
-					<!--Not logged in-->
-					<TextBlock Visibility="{Binding Path=LoggedInUser, Converter={StaticResource VisConv_Collapsed_Inv}}">
-						<Hyperlink FontWeight="Bold" Command="{Binding OnLogInClicked}">Log in</Hyperlink>
-					</TextBlock>
-				</TextBlock>
+					<Grid Grid.Row="1" ToolTip="User authentication allows enumerating files and downloading data directly through the proxy server">
+						<Grid.ColumnDefinitions>
+							<ColumnDefinition Width="auto"/>
+							<ColumnDefinition Width="*"/>
+						</Grid.ColumnDefinitions>
+						<TextBlock Grid.Column="0" Visibility="{Binding Path=ShouldShowLoginInfo, Converter={StaticResource VisConv_Collapsed}}">
+							<Run Text="[Experimental] User:"/>
+							<!--Logged in-->
+							<TextBlock Visibility="{Binding Path=LoggedInUser, Converter={StaticResource VisConv_Collapsed}}">
+								<Run FontWeight="Bold" Text="{Binding Path=LoggedInUser, Mode=OneWay}"/>
+								<Run Text=" - "/>
+								<Hyperlink Command="{Binding OnLogOutClicked}">Log out</Hyperlink>
+							</TextBlock>
+							<!--Not logged in-->
+							<TextBlock Visibility="{Binding Path=LoggedInUser, Converter={StaticResource VisConv_Collapsed_Inv}}">
+								<Hyperlink FontWeight="Bold" Command="{Binding OnLogInClicked}">Log in</Hyperlink>
+							</TextBlock>
+						</TextBlock>
+					</Grid>
 				</Grid>
 			</GroupBox>
 
@@ -242,7 +248,10 @@
 								TODO: expandable log window
 								<Expander VerticalAlignment="Center" Name="ApplicationLogExpander" IsExpanded="{Binding LogExpanded}" />
 							-->
-							<TextBlock VerticalAlignment="Center"><Run Text="Application Log"/></TextBlock>
+							<TextBlock VerticalAlignment="Center"><Run Text="Application Log"/>
+								<Run Text=" - " />
+								<Hyperlink Command="{Binding OnClearApplicationLogClicked}">Clear</Hyperlink>
+							</TextBlock>
 						</WrapPanel>
 					</GroupBox.Header>
 					<local:AutoScrollTextBox 
diff --git a/Engine/Source/Programs/UnsyncUI/UnsyncUI/MainWindow.xaml.cs b/Engine/Source/Programs/UnsyncUI/UnsyncUI/MainWindow.xaml.cs
--- a/Engine/Source/Programs/UnsyncUI/UnsyncUI/MainWindow.xaml.cs
+++ b/Engine/Source/Programs/UnsyncUI/UnsyncUI/MainWindow.xaml.cs
@@ -268,6 +268,8 @@ namespace UnsyncUI
 		public Command OnClearQueueClicked { get; }
 		public Command OnClearCompletedClicked { get; }
 
+		public Command OnClearApplicationLogClicked { get; }
+
 		public ObservableCollection<TabModel> Tabs { get; } = new ObservableCollection<TabModel>();
 		public ObservableCollection<JobModel> QueuedJobs { get; } = new ObservableCollection<JobModel>();
 
@@ -295,6 +297,7 @@ namespace UnsyncUI
 		{
 			OnClearQueueClicked = new Command(ClearQueue) { Enabled = true };
 			OnClearCompletedClicked = new Command(ClearCompleted) { Enabled = true };
+			OnClearApplicationLogClicked = new Command(ClearApplicationLog) { Enabled = true };
 
 			OnLogInClicked = new Command(LogIn) { Enabled = true };
 			OnLogOutClicked = new Command(LogOut) { Enabled = true };
@@ -330,6 +333,11 @@ namespace UnsyncUI
 			UpdateProgressState();
 		}
 
+		private void ClearApplicationLog()
+		{
+			App.Current.ClearApplicationLog();
+		}
+
 		public void OnBuildsSelected(IEnumerable<(SyncStartConfig Config, BuildPlatformModel Model)> selectedBuilds)
 		{
 			foreach (var build in selectedBuilds)
