Quick fix for null reference exception in Log calls through App.Current from external application contexts like LauncherTool

diff --git a/Engine/Source/Programs/UnsyncUI/UnsyncUI/Utils.cs b/Engine/Source/Programs/UnsyncUI/UnsyncUI/Utils.cs
--- a/Engine/Source/Programs/UnsyncUI/UnsyncUI/Utils.cs
+++ b/Engine/Source/Programs/UnsyncUI/UnsyncUI/Utils.cs
@@ -93,7 +93,7 @@ namespace UnsyncUI
 
         public static Task<IEnumerable<string>> EnumerateDirectoriesAsync(string path, CancellationToken token)
         {
-			App.Current.LogMessage($"Enumerating: {path}");
+			App.Current?.LogMessage($"Enumerating: {path}");
 
 			var tcs = new TaskCompletionSource<IEnumerable<string>>();
             Task.Run(() =>
@@ -143,7 +143,7 @@ namespace UnsyncUI
 						var dirs = Directory.EnumerateFiles(path).ToList();
 
 						timer.Stop();
-						App.Current.LogMessage($"Time: {timer.Elapsed.TotalSeconds:0.000} s - {path}");
+						App.Current?.LogMessage($"Time: {timer.Elapsed.TotalSeconds:0.000} s - {path}");
 
 						tcs.TrySetResult(dirs);
 					}
@@ -284,7 +284,7 @@ namespace UnsyncUI
 		public async IAsyncEnumerable<string> RunAsync([EnumeratorCancellation] CancellationToken cancelToken, bool ReadStdErr = true)
 		{
 			string processFileName = Path.GetFileName(proc.StartInfo.FileName);
-			App.Current.LogMessage($"Running: {processFileName} {proc.StartInfo.Arguments}");
+			App.Current?.LogMessage($"Running: {processFileName} {proc.StartInfo.Arguments}");
 
 			try
 			{
@@ -346,7 +346,7 @@ namespace UnsyncUI
 				ExitCode = proc.ExitCode;
 				proc.Dispose();
 
-				App.Current.LogDebug($"Finished: {processFileName} with exit code {ExitCode}");
+				App.Current?.LogDebug($"Finished: {processFileName} with exit code {ExitCode}");
 
 				cancelToken.ThrowIfCancellationRequested();
 			}
