UnsyncUI - Query mirrors by invoking `unsync.exe query mirrors` instead of custom HTTP request

diff --git a/Engine/Source/Programs/UnsyncUI/UnsyncUI/App.xaml.cs b/Engine/Source/Programs/UnsyncUI/UnsyncUI/App.xaml.cs
--- a/Engine/Source/Programs/UnsyncUI/UnsyncUI/App.xaml.cs
+++ b/Engine/Source/Programs/UnsyncUI/UnsyncUI/App.xaml.cs
@@ -87,27 +87,6 @@ namespace UnsyncUI
 
 			var oldWorkingDir = Environment.CurrentDirectory;
 
-			if (File.Exists(configFile))
-			{
-				// Move the working dir into the same as the config file.
-				// This allows relative paths to automatically be relative to the config.
-				var workingDir = Path.GetDirectoryName(Path.GetFullPath(configFile));
-				Environment.CurrentDirectory = workingDir;
-
-				try
-				{
-					Config = new Config(configFile);
-				}
-				catch (Exception ex)
-				{
-					MessageBox.Show($"Failed to load configuration from \"{configFile}\". {ex}");
-				}
-			}
-			else if (e.Args.Length > 0)
-			{
-				MessageBox.Show($"The configuration file \"{configFile}\" does not exist. Projects will not be available.");
-			}
-
 			string toolName = "unsync.exe";
 
 			// If config specifies the explicit tool path, then use that.
@@ -130,11 +109,26 @@ namespace UnsyncUI
 				UnsyncPath = oldWorkingDirTool;
 			}
 
-			if (UnsyncPath == null || !File.Exists(UnsyncPath))
+			if (File.Exists(configFile))
+			{
+				// Move the working dir into the same as the config file.
+				// This allows relative paths to automatically be relative to the config.
+				var workingDir = Path.GetDirectoryName(Path.GetFullPath(configFile));
+				Environment.CurrentDirectory = workingDir;
+
+				try
+				{
+					Config = new Config(configFile, UnsyncPath);
+				}
+				catch (Exception ex)
+				{
+					MessageBox.Show($"Failed to load configuration from \"{configFile}\".\n\n{ex}", "Fatal error");
+					Shutdown(1);
+				}
+			}
+			else if (e.Args.Length > 0)
 			{
-				string name = UnsyncPath ?? toolName;
-				MessageBox.Show($"Failed to find \"{name}\".");
-				Shutdown();
+				MessageBox.Show($"The configuration file \"{configFile}\" does not exist. Projects will not be available.");
 			}
 
 			UserConfig = UserPreferences.Load();
@@ -151,7 +145,10 @@ namespace UnsyncUI
 
 		protected override void OnExit(ExitEventArgs e)
 		{
-			UserConfig.Save();
+			if (UserConfig != null)
+			{
+				UserConfig.Save();
+			}
 			base.OnExit(e);
 		}
 	}
diff --git a/Engine/Source/Programs/UnsyncUI/UnsyncUI/Config.cs b/Engine/Source/Programs/UnsyncUI/UnsyncUI/Config.cs
--- a/Engine/Source/Programs/UnsyncUI/UnsyncUI/Config.cs
+++ b/Engine/Source/Programs/UnsyncUI/UnsyncUI/Config.cs
@@ -12,23 +12,10 @@ using System.Xml.Linq;
 using System.Net;
 using System.Text.Json;
 
-// warning SYSLIB0014: 'WebRequest.Create(string)' is obsolete: 'WebRequest, HttpWebRequest, ServicePoint, and WebClient are obsolete. Use HttpClient instead.' (https://aka.ms/dotnet-warnings/SYSLIB0014)
-#pragma warning disable SYSLIB0014
-
 namespace UnsyncUI
 {
 	public sealed class Config
 	{
-		// Structure that represents mirror server entry in the list from /api/v1/mirrors endpoint
-		private class JsonMirrorDesc
-		{
-			public String name { get; set; }
-			public String address { get; set; }
-			public int port { get; set; } = 0;
-			public String description { get; set; }
-			public String parent { get; set; }
-		}
-
 		public sealed class Proxy
 		{
 			public string Name { get; set; }
@@ -228,13 +215,22 @@ namespace UnsyncUI
 
 		internal string loggedInUser;
 
-		public Config(string filename)
+		public Config(string filename, string DefaultUnsyncPath = null)
 		{
 			var rootNode = XDocument.Load(filename).Root;
 
 			UnsyncPath = rootNode.Attribute("path")?.Value;
+			if (UnsyncPath == null)
+			{
+				UnsyncPath = DefaultUnsyncPath;
+			}
+
+			if (UnsyncPath == null)
+			{
+				throw new Exception("Path to unsync.exe was not provided.");
+			}
 
-			if (UnsyncPath != null && !File.Exists(UnsyncPath))
+			if (!File.Exists(UnsyncPath))
 			{
 				throw new Exception("Unable to find unsync.exe binary specified in config file.");
 			}
@@ -281,8 +277,6 @@ namespace UnsyncUI
 		// Returns list of mirrors and the seed proxy server that was used to get it or null
 		private (List<Proxy>, Proxy) DiscoverProxies(List<Proxy> SeedServers)
 		{
-			int DefaultPort = 53841;
-
 			foreach (Proxy SeedServer in SeedServers)
 			{
 				if (SeedServer.Path == null)
@@ -292,60 +286,30 @@ namespace UnsyncUI
 
 				try
 				{
-					String Url = SeedServer.Path;
-					if (!Url.Contains(":"))
-					{
-						Url += ":" + DefaultPort.ToString();
-					}
+					UnsyncQueryUtil QueryUtil = new UnsyncQueryUtil(UnsyncPath, SeedServer.Path);
 
-					if (!Url.StartsWith("http://"))
-					{
-						Url = "http://" + Url;
-					}
+					var ParsedProxies = new List<Proxy>();
 
-					Url += "/api/v1/mirrors";
-
-					var Request = WebRequest.Create(Url);
-					Request.Timeout = 2500;
-					var Response = (HttpWebResponse)Request.GetResponse();
-					if (Response.StatusCode == HttpStatusCode.OK)
+					foreach (var Mirror in QueryUtil.Mirrors())
 					{
-						var Reader = new StreamReader(Response.GetResponseStream());
-						var Body = Reader.ReadToEnd();
-						var ParsedList = JsonSerializer.Deserialize<List<JsonMirrorDesc>>(Body);
-
-						var ParsedProxies = new List<Proxy>();
-						foreach (var ParsedProxy in ParsedList)
+						if (!Mirror.ok || Mirror.ping == 0)
 						{
-							if (ParsedProxy.address == null)
-							{
-								continue;
-							}
-
-							var ConvertedProxy = new Proxy();
-							ConvertedProxy.Path = ParsedProxy.address;
-
-							if (ParsedProxy.description != null)
-							{
-								ConvertedProxy.Name = ParsedProxy.description;
-							}
-							else if (ParsedProxy.name != null)
-							{
-								ConvertedProxy.Name = ParsedProxy.name;
-							}
-
-							if (ParsedProxy.port != 0)
-							{
-								ConvertedProxy.Path += ":" + ParsedProxy.port.ToString();
-							}
-
-							ParsedProxies.Add(ConvertedProxy);
+							continue;
 						}
 
-						if (ParsedProxies.Count != 0)
+						var ConvertedProxy = new Proxy();
+						ConvertedProxy.Name = (Mirror.description != null && Mirror.description.Length > 0) ? Mirror.description : Mirror.name;
+						ConvertedProxy.Path = Mirror.address;
+						if (Mirror.port != 0)
 						{
-							return (ParsedProxies, SeedServer);
+							ConvertedProxy.Path += ":" + Mirror.port.ToString();
 						}
+						ParsedProxies.Add(ConvertedProxy);
+					}
+
+					if (ParsedProxies.Count != 0)
+					{
+						return (ParsedProxies, SeedServer);
 					}
 				}
 				catch (Exception)
diff --git a/Engine/Source/Programs/UnsyncUI/UnsyncUI/Query.cs b/Engine/Source/Programs/UnsyncUI/UnsyncUI/Query.cs
--- a/Engine/Source/Programs/UnsyncUI/UnsyncUI/Query.cs
+++ b/Engine/Source/Programs/UnsyncUI/UnsyncUI/Query.cs
@@ -45,6 +45,16 @@ namespace UnsyncUI
 		public List<String> groups { get; set; }
 	}
 
+	public class UnsyncMirrorDesc
+	{
+		public String name { get; set; }
+		public String description { get; set; }
+		public String address { get; set; }
+		public uint port { get; set; }
+		public uint ping { get; set; }
+		public bool ok { get; set; }
+	}
+
 	public class UnsyncQueryUtil
 	{
 		UnsyncQueryConfig Config;
@@ -60,6 +70,26 @@ namespace UnsyncUI
 			Config.proxyAddress = proxyAddress;
 		}
 
+		public List<UnsyncMirrorDesc> Mirrors()
+		{
+			String argsStr = $"query mirrors --proxy {Config.proxyAddress}";
+			AsyncProcess proc = new AsyncProcess(Config.unsyncPath, argsStr);
+			CancellationToken cancellationToken = new CancellationToken();
+
+			var responseJson = "";
+
+			var QueryTask = Task.Run(async () => {
+				// TODO: read stderr stream and somehow report status/errors
+				await foreach (var str in proc.RunAsync(cancellationToken, false /*ReadStdErr*/))
+				{
+					responseJson += str;
+				}
+			});
+			QueryTask.Wait();
+
+			return JsonSerializer.Deserialize<List<UnsyncMirrorDesc>>(responseJson);
+		}
+
 		public LoginQueryResult Login()
 		{
 			String argsStr = $"login --decode --proxy {Config.proxyAddress}";
diff --git a/Engine/Source/Programs/UnsyncUI/UnsyncUI/UnsyncUI.csproj b/Engine/Source/Programs/UnsyncUI/UnsyncUI/UnsyncUI.csproj
--- a/Engine/Source/Programs/UnsyncUI/UnsyncUI/UnsyncUI.csproj
+++ b/Engine/Source/Programs/UnsyncUI/UnsyncUI/UnsyncUI.csproj
@@ -8,7 +8,7 @@
 		<ApplicationManifest>app.manifest</ApplicationManifest>
         <IncludeAllContentForSelfExtract>true</IncludeAllContentForSelfExtract>
         <ApplicationIcon>unsync_48px.ico</ApplicationIcon>
-        <Version>$(VersionPrefix)1.0.21</Version>
+        <Version>$(VersionPrefix)1.0.22</Version>
 	</PropertyGroup>
 
 </Project>
