UnsyncUI - Add top-level unhandled exception handler and report stderr output from queries

diff --git a/Engine/Source/Programs/UnsyncUI/UnsyncUI/App.xaml.cs b/Engine/Source/Programs/UnsyncUI/UnsyncUI/App.xaml.cs
--- a/Engine/Source/Programs/UnsyncUI/UnsyncUI/App.xaml.cs
+++ b/Engine/Source/Programs/UnsyncUI/UnsyncUI/App.xaml.cs
@@ -16,6 +16,12 @@ namespace UnsyncUI
 	/// </summary>
 	public partial class App : Application
     {
+		public App() : base()
+		{
+			Dispatcher.UnhandledException += (sender, args) => UnhandledException(sender, args.Exception.ToString());
+			AppDomain.CurrentDomain.UnhandledException += (sender, args) => UnhandledException(sender, args.ExceptionObject.ToString());
+		}
+
 		public new static App Current => Application.Current as App;
 
 		public Config Config { get; private set; }
@@ -24,7 +30,7 @@ namespace UnsyncUI
 		public string DefaultSearchTerms { get; private set; } = "";
 
 		internal bool EnableExperimentalFeatures = false;
-		internal bool EnableUserAuthentication = true;
+		internal bool EnableUserAuthentication = false;
 
 		internal string ApplicationLog { get; private set; } = "";
 
@@ -135,6 +141,21 @@ namespace UnsyncUI
 
 			if (Config != null)
 			{
+				/*
+				if (Config.RootProxy != null)
+				{
+					App.Current.LogMessage($"Unsync server address: {Config.RootProxy.Path}");
+					EnableUserAuthentication = true;
+				}
+				else
+				{
+					App.Current.LogMessage($"Unsync server address is not configured");
+					EnableUserAuthentication = false;
+				}
+				*/
+
+				EnableUserAuthentication = true;
+
 				Config.UnsyncPath = UnsyncPath;
 				Config.EnableExperimentalFeatures = EnableExperimentalFeatures;
 				Config.EnableUserAuthentication = EnableUserAuthentication;
@@ -151,6 +172,12 @@ namespace UnsyncUI
 			}
 			base.OnExit(e);
 		}
+
+		static void UnhandledException(object sender, string e)
+		{
+			MessageBox.Show($"Unhandled exception:\n{e}", "Fatal error", MessageBoxButton.OK, MessageBoxImage.Error);
+			Environment.Exit(1);
+		}
 	}
 
 	public sealed class UserPreferences
diff --git a/Engine/Source/Programs/UnsyncUI/UnsyncUI/Config.cs b/Engine/Source/Programs/UnsyncUI/UnsyncUI/Config.cs
--- a/Engine/Source/Programs/UnsyncUI/UnsyncUI/Config.cs
+++ b/Engine/Source/Programs/UnsyncUI/UnsyncUI/Config.cs
@@ -243,16 +243,29 @@ namespace UnsyncUI
 
 			List<Proxy> ConfigProxies = new List<Proxy>();
 
-			ConfigProxies.AddRange(rootNode.Element("proxies").Elements("proxy").Select(p => new Proxy()
+			var proxiesConfigNode = rootNode.Element("proxies");
+			if (proxiesConfigNode != null)
 			{
-				Name = p.Attribute("name")?.Value,
-				Path = p.Attribute("path")?.Value
-			}));
+				ConfigProxies.AddRange(proxiesConfigNode.Elements("proxy").Select(p => new Proxy()
+				{
+					Name = p.Attribute("name")?.Value,
+					Path = p.Attribute("path")?.Value
+				}));
+			}
+
+			// Set a default root proxy
+			if (ConfigProxies.Count > 0)
+			{
+				RootProxy = ConfigProxies.First();
+			}
 
 			// Auto-discover proxies
 			(List<Proxy> DiscoveredProxies, Proxy DiscoveredRootProxy) = DiscoverProxies(ConfigProxies);
 
-			RootProxy = DiscoveredRootProxy;
+			if (DiscoveredRootProxy != null)
+			{
+				RootProxy = DiscoveredRootProxy;
+			}
 
 			if (DiscoveredProxies == null)
 			{
@@ -312,8 +325,9 @@ namespace UnsyncUI
 						return (ParsedProxies, SeedServer);
 					}
 				}
-				catch (Exception)
+				catch (Exception ex)
 				{
+					App.Current.LogError("Exception during unsync query: " + ex.Message);
 					continue;
 				}
 			}
diff --git a/Engine/Source/Programs/UnsyncUI/UnsyncUI/MainWindow.xaml.cs b/Engine/Source/Programs/UnsyncUI/UnsyncUI/MainWindow.xaml.cs
--- a/Engine/Source/Programs/UnsyncUI/UnsyncUI/MainWindow.xaml.cs
+++ b/Engine/Source/Programs/UnsyncUI/UnsyncUI/MainWindow.xaml.cs
@@ -236,7 +236,6 @@ namespace UnsyncUI
 			}
 			catch (Exception ex)
 			{
-				// TODO: add global status/log window
 				App.Current.LogError($"Login failed with exception: {ex}");
 			}
 		}
diff --git a/Engine/Source/Programs/UnsyncUI/UnsyncUI/Query.cs b/Engine/Source/Programs/UnsyncUI/UnsyncUI/Query.cs
--- a/Engine/Source/Programs/UnsyncUI/UnsyncUI/Query.cs
+++ b/Engine/Source/Programs/UnsyncUI/UnsyncUI/Query.cs
@@ -70,46 +70,50 @@ namespace UnsyncUI
 			Config.proxyAddress = proxyAddress;
 		}
 
-		public List<UnsyncMirrorDesc> Mirrors()
+		private string RunCommand(string argsStr)
 		{
-			String argsStr = $"query mirrors --proxy {Config.proxyAddress}";
 			AsyncProcess proc = new AsyncProcess(Config.unsyncPath, argsStr);
 			CancellationToken cancellationToken = new CancellationToken();
 
-			var responseJson = "";
+			var response = "";
+			var diagnostics = "";
 
 			var QueryTask = Task.Run(async () => {
-				// TODO: read stderr stream and somehow report status/errors
-				await foreach (var str in proc.RunAsync(cancellationToken, false /*ReadStdErr*/))
+				await foreach (var (str, kind) in proc.RunAsyncStreams(cancellationToken))
 				{
-					responseJson += str;
+					if (kind == AsyncProcess.StreamKind.StdOut)
+					{
+						response += str;
+					}
+					else
+					{
+						diagnostics += str;
+					}
 				}
 			});
+
 			QueryTask.Wait();
 
+			if (proc.ExitCode == 0)
+			{
+				return response;
+			}
+			else
+			{
+				throw new Exception($"Exit code {proc.ExitCode}\n{diagnostics}");
+			}
+		}
+
+		public List<UnsyncMirrorDesc> Mirrors()
+		{
+			string responseJson = RunCommand($"query mirrors --proxy {Config.proxyAddress}");
 			return JsonSerializer.Deserialize<List<UnsyncMirrorDesc>>(responseJson);
 		}
 
 		public LoginQueryResult Login()
 		{
-			String argsStr = $"login --decode --proxy {Config.proxyAddress}";
-			AsyncProcess proc = new AsyncProcess(Config.unsyncPath, argsStr);
-			CancellationToken cancellationToken = new CancellationToken();
-
-			var responseJson = "";
-
-			var LoginTask = Task.Run(async () => {
-				// TODO: read stderr stream and somehow report status/errors
-				await foreach (var str in proc.RunAsync(cancellationToken, false /*ReadStdErr*/))
-				{
-					responseJson += str;
-				}
-			});
-			LoginTask.Wait();
-
-			LoginQueryResult queryResult = JsonSerializer.Deserialize<LoginQueryResult>(responseJson);
-
-			return queryResult;
+			string responseJson = RunCommand($"login --decode --proxy {Config.proxyAddress}");
+			return JsonSerializer.Deserialize<LoginQueryResult>(responseJson);
 		}
 	}
 
@@ -215,7 +219,7 @@ namespace UnsyncUI
 					}
 					catch (Exception ex)
 					{
-						App.Current.LogError("Exception while parsing unsync query JSON: " + ex.Message);
+						App.Current.LogError("Exception during unsync query: " + ex.Message);
 					}
 				}
 			}
diff --git a/Engine/Source/Programs/UnsyncUI/UnsyncUI/UnsyncUI.csproj b/Engine/Source/Programs/UnsyncUI/UnsyncUI/UnsyncUI.csproj
--- a/Engine/Source/Programs/UnsyncUI/UnsyncUI/UnsyncUI.csproj
+++ b/Engine/Source/Programs/UnsyncUI/UnsyncUI/UnsyncUI.csproj
@@ -8,7 +8,7 @@
 		<ApplicationManifest>app.manifest</ApplicationManifest>
         <IncludeAllContentForSelfExtract>true</IncludeAllContentForSelfExtract>
         <ApplicationIcon>unsync_48px.ico</ApplicationIcon>
-        <Version>$(VersionPrefix)1.0.22</Version>
+        <Version>$(VersionPrefix)1.0.23</Version>
 	</PropertyGroup>
 
 </Project>
diff --git a/Engine/Source/Programs/UnsyncUI/UnsyncUI/Utils.cs b/Engine/Source/Programs/UnsyncUI/UnsyncUI/Utils.cs
--- a/Engine/Source/Programs/UnsyncUI/UnsyncUI/Utils.cs
+++ b/Engine/Source/Programs/UnsyncUI/UnsyncUI/Utils.cs
@@ -281,7 +281,24 @@ namespace UnsyncUI
 			};
 		}
 
+		public enum StreamKind
+		{
+			StdOut,
+			StdErr,
+		}
+
 		public async IAsyncEnumerable<string> RunAsync([EnumeratorCancellation] CancellationToken cancelToken, bool ReadStdErr = true)
+		{
+			await foreach (var (s, k) in RunAsyncStreams(cancelToken))
+			{
+				if (k == StreamKind.StdOut || (k==StreamKind.StdErr && ReadStdErr))
+				{
+					yield return s;
+				}
+			}
+		}
+
+		public async IAsyncEnumerable<(string, StreamKind)> RunAsyncStreams([EnumeratorCancellation] CancellationToken cancelToken)
 		{
 			string processFileName = Path.GetFileName(proc.StartInfo.FileName);
 			App.Current?.LogMessage($"Running: {processFileName} {proc.StartInfo.Arguments}");
@@ -299,9 +316,9 @@ namespace UnsyncUI
 					proc.Start();
 					ProcessTracker.Attach(proc);
 
-					var pipe = new BufferBlock<string>();
+					var pipe = new BufferBlock<(string, StreamKind)>();
 
-					async Task ReadStream(Stream stream, bool ShouldPost)
+					async Task ReadStream(Stream stream, StreamKind Kind)
 					{
 						var block = new byte[4096];
 						var mem = new Memory<byte>(block);
@@ -314,17 +331,14 @@ namespace UnsyncUI
 								break;
 							}
 
-							if (ShouldPost)
-							{
-								// @todo: this won't handle UTF-8 encoding if a char is split across the read boundary.
-								string decodedString = Encoding.UTF8.GetString(mem.Span.Slice(0, bytesRead));
-								pipe.Post(decodedString);
-							}
+							// @todo: this won't handle UTF-8 encoding if a char is split across the read boundary.
+							string decodedString = Encoding.UTF8.GetString(mem.Span.Slice(0, bytesRead));
+							pipe.Post((decodedString, Kind));
 						}
 					}
 
-					var stdoutTask = ReadStream(proc.StandardOutput.BaseStream, true);
-					var stderrTask = ReadStream(proc.StandardError.BaseStream, ReadStdErr);
+					var stdoutTask = ReadStream(proc.StandardOutput.BaseStream, StreamKind.StdOut);
+					var stderrTask = ReadStream(proc.StandardError.BaseStream, StreamKind.StdErr);
 
 					var completionTask = Task.Run(async () =>
 					{
@@ -334,8 +348,8 @@ namespace UnsyncUI
 
 					while (await pipe.OutputAvailableAsync())
 					{
-						string receivedString = pipe.Receive();
-						yield return receivedString;
+						var (receivedString, receivedKind) = pipe.Receive();
+						yield return (receivedString, receivedKind);
 					}
 
 					await completionTask;
