UnsyncUI - Fix crash when running without config file

diff --git a/Engine/Source/Programs/UnsyncUI/UnsyncUI/MainWindow.xaml.cs b/Engine/Source/Programs/UnsyncUI/UnsyncUI/MainWindow.xaml.cs
--- a/Engine/Source/Programs/UnsyncUI/UnsyncUI/MainWindow.xaml.cs
+++ b/Engine/Source/Programs/UnsyncUI/UnsyncUI/MainWindow.xaml.cs
@@ -207,7 +207,7 @@ namespace UnsyncUI
 
 		public string LoggedInUser
 		{
-			get => Config.loggedInUser;
+			get => Config?.loggedInUser;
 			set 
 			{
 				App.Current.UserConfig.LogInOnStartup = (value != null);
@@ -307,12 +307,11 @@ namespace UnsyncUI
 				{
 					Tabs.Add(new ProjectModel(p, OnBuildsSelected));
 				}
-			}
 
-			if (Config.EnableUserAuthentication
-				&& App.Current.UserConfig.LogInOnStartup)
-			{
-				LogIn();
+				if (Config.EnableUserAuthentication && App.Current.UserConfig.LogInOnStartup)
+				{
+					LogIn();
+				}
 			}
 
 			Tabs.Add(new CustomModel(OnBuildsSelected));
