[UnsyncUI] Fix user-unhandled exception crash from async task

#rnx

diff --git a/Engine/Source/Programs/UnsyncUI/UnsyncUI/BuildSelectorControl.xaml.cs b/Engine/Source/Programs/UnsyncUI/UnsyncUI/BuildSelectorControl.xaml.cs
--- a/Engine/Source/Programs/UnsyncUI/UnsyncUI/BuildSelectorControl.xaml.cs
+++ b/Engine/Source/Programs/UnsyncUI/UnsyncUI/BuildSelectorControl.xaml.cs
@@ -151,38 +151,45 @@ namespace UnsyncUI
 			if (!currentDir.Parse(path, ref template))
 				return;
 
-			foreach (var childDir in await dirEnum.EnumerateDirectories(path, cancellationToken.Token))
+			try
 			{
-				if (System.IO.Path.GetFileName(childDir) == ".unsync")
+				foreach (var childDir in await dirEnum.EnumerateDirectories(path, cancellationToken.Token))
 				{
-					List<Config.BuildTemplate> fileGroups = null;
-					if (currentDir.FileGroups?.Any() == true)
-					{
-						var files = await dirEnum.EnumerateFiles(path, cancellationToken.Token);
-						currentDir.ParseFileGroups(files.ToList(), template, out fileGroups);
-					}
-					if (fileGroups?.Any() == true)
+					if (System.IO.Path.GetFileName(childDir) == ".unsync")
 					{
-						foreach (var fileGroup in fileGroups)
+						List<Config.BuildTemplate> fileGroups = null;
+						if (currentDir.FileGroups?.Any() == true)
+						{
+							var files = await dirEnum.EnumerateFiles(path, cancellationToken.Token);
+							currentDir.ParseFileGroups(files.ToList(), template, out fileGroups);
+						}
+						if (fileGroups?.Any() == true)
 						{
-							Platforms.Add(new BuildPlatformModel(this, fileGroup.Platform, path, fileGroup.Flavor, fileGroup.Include));
+							foreach (var fileGroup in fileGroups)
+							{
+								Platforms.Add(new BuildPlatformModel(this, fileGroup.Platform, path, fileGroup.Flavor, fileGroup.Include));
+							}
+							Platforms.Add(new BuildPlatformModel(this, template.Platform, path, "All", null));
+						}
+						else
+						{
+							// This folder is a valid build
+							Platforms.Add(new BuildPlatformModel(this, template.Platform, path, template.Flavor, null));
 						}
-						Platforms.Add(new BuildPlatformModel(this, template.Platform, path, "All", null));
 					}
 					else
 					{
-						// This folder is a valid build
-						Platforms.Add(new BuildPlatformModel(this, template.Platform, path, template.Flavor, null));
-					}
-				}
-				else
-				{
-					foreach (var subDir in currentDir.SubDirectories)
-					{
-						await EnumeratePlatforms(dirEnum, subDir, childDir, template, cancellationToken);
+						foreach (var subDir in currentDir.SubDirectories)
+						{
+							await EnumeratePlatforms(dirEnum, subDir, childDir, template, cancellationToken);
+						}
 					}
 				}
 			}
+			catch (Exception ex)
+			{
+				App.Current.LogMessage($"Platforms enumeration failed with exception: {ex}");
+			}
 		}
 
 		public async void PopulatePlatforms()
