[UnsyncUI] Add null check to catch logging
This prevents further crashes from this function being called in an Application context that is not UnsyncUI

#rnx

diff --git a/Engine/Source/Programs/UnsyncUI/UnsyncUI/BuildSelectorControl.xaml.cs b/Engine/Source/Programs/UnsyncUI/UnsyncUI/BuildSelectorControl.xaml.cs
--- a/Engine/Source/Programs/UnsyncUI/UnsyncUI/BuildSelectorControl.xaml.cs
+++ b/Engine/Source/Programs/UnsyncUI/UnsyncUI/BuildSelectorControl.xaml.cs
@@ -188,7 +188,7 @@ namespace UnsyncUI
 			}
 			catch (Exception ex)
 			{
-				App.Current.LogMessage($"Platforms enumeration failed with exception: {ex}");
+				App.Current?.LogMessage($"Platforms enumeration failed with exception: {ex}");
 			}
 		}
 
