UnsyncUI - Fix few possible cases of null object dereferences

* UnsyncDirectoryEnumerator.Entries
* BuildModel.Suffix

diff --git a/Engine/Source/Programs/UnsyncUI/UnsyncUI/BuildSelectorControl.xaml.cs b/Engine/Source/Programs/UnsyncUI/UnsyncUI/BuildSelectorControl.xaml.cs
--- a/Engine/Source/Programs/UnsyncUI/UnsyncUI/BuildSelectorControl.xaml.cs
+++ b/Engine/Source/Programs/UnsyncUI/UnsyncUI/BuildSelectorControl.xaml.cs
@@ -116,7 +116,7 @@ namespace UnsyncUI
 		private CancellationTokenSource cts;
 		public string Stream => Template.Stream;
 		public string CL => Template.CL;
-		public string Suffix => Template.Suffix;
+		public string Suffix => Template.Suffix == null ? "" : Template.Suffix;
 
 		public Config.Directory RootDir { get; }
 		public string Path { get; }
diff --git a/Engine/Source/Programs/UnsyncUI/UnsyncUI/Query.cs b/Engine/Source/Programs/UnsyncUI/UnsyncUI/Query.cs
--- a/Engine/Source/Programs/UnsyncUI/UnsyncUI/Query.cs
+++ b/Engine/Source/Programs/UnsyncUI/UnsyncUI/Query.cs
@@ -157,6 +157,8 @@ namespace UnsyncUI
 		{
 			if (!Initialized)
 			{
+				Entries = new List<Entry>();
+
 				if (ProjectSchema != null)
 				{
 					await InitProject(cancellationToken);
@@ -179,11 +181,6 @@ namespace UnsyncUI
 
 		private void ProcessQueryResult(SearchQueryResult queryResult)
 		{
-			if (Entries == null)
-			{
-				Entries = new List<Entry>();
-			}
-
 			foreach (var queryEntry in queryResult.entries)
 			{
 				Entry entry = new Entry();
