UnsyncUI - Fix starting downloads for builds that are only different by flavor/config

diff --git a/Engine/Source/Programs/UnsyncUI/UnsyncUI/JobControl.xaml.cs b/Engine/Source/Programs/UnsyncUI/UnsyncUI/JobControl.xaml.cs
--- a/Engine/Source/Programs/UnsyncUI/UnsyncUI/JobControl.xaml.cs
+++ b/Engine/Source/Programs/UnsyncUI/UnsyncUI/JobControl.xaml.cs
@@ -263,6 +263,7 @@ namespace UnsyncUI
 		{
 			return (Build.Build.Path == otherJob.Build.Build.Path)
 				&& (Build.Platform == otherJob.Build.Platform)
+				&& (Build.Flavor == otherJob.Build.Flavor)
 				&& (DstPathBase == otherJob.DstPathBase);
 		}
 
