UnsyncUI - Add protocol field to config file, allowing for behavior differences between unsync/filesystem and horde back-ends

diff --git a/Engine/Source/Programs/UnsyncUI/UnsyncUI/Config.cs b/Engine/Source/Programs/UnsyncUI/UnsyncUI/Config.cs
--- a/Engine/Source/Programs/UnsyncUI/UnsyncUI/Config.cs
+++ b/Engine/Source/Programs/UnsyncUI/UnsyncUI/Config.cs
@@ -20,6 +20,13 @@ namespace UnsyncUI
 		{
 			public string Name { get; set; }
 			public string Path { get; set; }
+			public string Protocol { get; set; }
+
+			public UnsyncServerConfig GetServerConfig()
+			{
+				return new UnsyncServerConfig { address=Path, protocol=Protocol };
+			}
+
 		}
 
 		internal struct BuildTemplate
@@ -235,12 +242,6 @@ namespace UnsyncUI
 				throw new Exception("Unable to find unsync.exe binary specified in config file.");
 			}
 
-			Proxies.Add(new Proxy()
-			{
-				Name = "(none)",
-				Path = null
-			});
-
 			List<Proxy> ConfigProxies = new List<Proxy>();
 
 			var proxiesConfigNode = rootNode.Element("proxies");
@@ -249,7 +250,8 @@ namespace UnsyncUI
 				ConfigProxies.AddRange(proxiesConfigNode.Elements("proxy").Select(p => new Proxy()
 				{
 					Name = p.Attribute("name")?.Value,
-					Path = p.Attribute("path")?.Value
+					Path = p.Attribute("path")?.Value,
+					Protocol = p.Attribute("protocol")?.Value
 				}));
 			}
 
@@ -259,21 +261,39 @@ namespace UnsyncUI
 				RootProxy = ConfigProxies.First();
 			}
 
-			// Auto-discover proxies
-			(List<Proxy> DiscoveredProxies, Proxy DiscoveredRootProxy) = DiscoverProxies(ConfigProxies);
-
-			if (DiscoveredRootProxy != null)
+			if (RootProxy != null)
 			{
-				RootProxy = DiscoveredRootProxy;
-			}
+				// Horde requires a server connection, but unsync / native file system does not
+				if (RootProxy.Protocol == "horde")
+				{
+					Proxies.AddRange(ConfigProxies);
+				}
+				else
+				{
+					Proxies.Add(new Proxy()
+					{
+						Name = "(none)",
+						Path = null,
+						Protocol = null,
+					});
 
-			if (DiscoveredProxies == null)
-			{
-				Proxies.AddRange(ConfigProxies);
-			}
-			else
-			{
-				Proxies.AddRange(DiscoveredProxies);
+					// Auto-discover proxies
+					(List<Proxy> DiscoveredProxies, Proxy DiscoveredRootProxy) = DiscoverProxies(ConfigProxies);
+
+					if (DiscoveredRootProxy != null)
+					{
+						RootProxy = DiscoveredRootProxy;
+					}
+
+					if (DiscoveredProxies == null)
+					{
+						Proxies.AddRange(ConfigProxies);
+					}
+					else
+					{
+						Proxies.AddRange(DiscoveredProxies);
+					}
+				}
 			}
 
 			Projects = rootNode.Element("projects").Elements("project").Select(p => new Project()
@@ -299,7 +319,7 @@ namespace UnsyncUI
 
 				try
 				{
-					UnsyncQueryUtil QueryUtil = new UnsyncQueryUtil(UnsyncPath, SeedServer.Path);
+					UnsyncQueryUtil QueryUtil = new UnsyncQueryUtil(UnsyncPath, SeedServer.GetServerConfig());
 
 					var ParsedProxies = new List<Proxy>();
 
@@ -338,7 +358,7 @@ namespace UnsyncUI
 		private UnsyncQueryConfig CreateUnsyncQueryConfig()
 		{
 			UnsyncQueryConfig unsyncConfig = new UnsyncQueryConfig();
-			unsyncConfig.proxyAddress = RootProxy.Path;
+			unsyncConfig.server = RootProxy.GetServerConfig();
 			unsyncConfig.unsyncPath = UnsyncPath;
 			return unsyncConfig;
 		}
diff --git a/Engine/Source/Programs/UnsyncUI/UnsyncUI/JobControl.xaml b/Engine/Source/Programs/UnsyncUI/UnsyncUI/JobControl.xaml
--- a/Engine/Source/Programs/UnsyncUI/UnsyncUI/JobControl.xaml
+++ b/Engine/Source/Programs/UnsyncUI/UnsyncUI/JobControl.xaml
@@ -60,7 +60,7 @@
 			</Grid.RowDefinitions>
 			<TextBlock Grid.Column="0" Grid.Row="0" Text="Source:"/>
 			<TextBlock Grid.Column="0" Grid.Row="1" Text="Destination:"/>
-			<TextBlock Grid.Column="0" Grid.Row="2" Text="Proxy:"/>
+			<TextBlock Grid.Column="0" Grid.Row="2" Text="Server:"/>
 
 			<TextBlock Grid.Column="2" Grid.Row="0">
                 <Hyperlink Command="{Binding OnSourcePathClicked}">
diff --git a/Engine/Source/Programs/UnsyncUI/UnsyncUI/JobControl.xaml.cs b/Engine/Source/Programs/UnsyncUI/UnsyncUI/JobControl.xaml.cs
--- a/Engine/Source/Programs/UnsyncUI/UnsyncUI/JobControl.xaml.cs
+++ b/Engine/Source/Programs/UnsyncUI/UnsyncUI/JobControl.xaml.cs
@@ -67,14 +67,14 @@ namespace UnsyncUI
 		public string ScavengePath { get; }
 		public string DstPathBase => Path.GetDirectoryName(DstPath);
         public bool DryRun { get; }
-        public string Proxy { get; }
+        public UnsyncServerConfig Server { get; }
 		public string AdditionalArgs { get; }
 		public string[] IncludeFilter { get; }
 		public string[] Exclusions { get; }
 
-        public string ProxyStr => string.IsNullOrWhiteSpace(Proxy) ? "(none)" : Proxy;
+        public string ProxyStr => string.IsNullOrWhiteSpace(Server?.address) ? "(none)" : Server.address;
 
-        public string Name => $"{Build.Name}{(DryRun ? " - Dry Run" : "")}";
+		public string Name => $"{Build.Name}{(DryRun ? " - Dry Run" : "")}";
 
 		private double totalProgress = 0;
         public double TotalProgress
@@ -171,7 +171,7 @@ namespace UnsyncUI
 			string dstPath,
 			string scavengePath,
 			bool dryRun,
-			string proxy,
+			UnsyncServerConfig server,
 			string additionalArgs,
 			string[] exclusions,
 			Action<JobModel> onCompletion,
@@ -182,7 +182,7 @@ namespace UnsyncUI
             DstPath = dstPath?.TrimEnd('\\');
 			ScavengePath = scavengePath?.TrimEnd('\\');
             DryRun = dryRun;
-            Proxy = proxy?.TrimEnd('\\');
+            Server = server;
 			AdditionalArgs = additionalArgs;
 			IncludeFilter = build.Include?.Split(',', StringSplitOptions.RemoveEmptyEntries)
 				.Select(Entry => Entry.Trim())
@@ -366,9 +366,9 @@ namespace UnsyncUI
 					args.Add("-d");
 				}
 
-				if (!string.IsNullOrWhiteSpace(Proxy))
+				if (Server != null)
 				{
-					args.Add($"--proxy \"{Proxy}\"");
+					args.Add(Server.GetCommandLineArgs());
 
 					if (App.Current.Config.EnableUserAuthentication
 						&& App.Current.Config.loggedInUser != null)
diff --git a/Engine/Source/Programs/UnsyncUI/UnsyncUI/MainWindow.xaml b/Engine/Source/Programs/UnsyncUI/UnsyncUI/MainWindow.xaml
--- a/Engine/Source/Programs/UnsyncUI/UnsyncUI/MainWindow.xaml
+++ b/Engine/Source/Programs/UnsyncUI/UnsyncUI/MainWindow.xaml
@@ -85,7 +85,7 @@
 				<RowDefinition Height="auto" />
 				<RowDefinition Height="*" />
 			</Grid.RowDefinitions>
-			<GroupBox Grid.Row="0" Header="Proxy">
+			<GroupBox Grid.Row="0" Header="Server">
 				<Grid>
 					<Grid.RowDefinitions>
 						<RowDefinition Height="auto" />
@@ -111,7 +111,7 @@
 							<ColumnDefinition Width="*"/>
 						</Grid.ColumnDefinitions>
 						<TextBlock Grid.Column="0" Visibility="{Binding Path=ShouldShowLoginInfo, Converter={StaticResource VisConv_Collapsed}}">
-							<Run Text="[Experimental] User:"/>
+							<Run Text="User:"/>
 							<!--Logged in-->
 							<TextBlock Visibility="{Binding Path=LoggedInUser, Converter={StaticResource VisConv_Collapsed}}">
 								<Run FontWeight="Bold" Text="{Binding Path=LoggedInUser, Mode=OneWay}"/>
diff --git a/Engine/Source/Programs/UnsyncUI/UnsyncUI/MainWindow.xaml.cs b/Engine/Source/Programs/UnsyncUI/UnsyncUI/MainWindow.xaml.cs
--- a/Engine/Source/Programs/UnsyncUI/UnsyncUI/MainWindow.xaml.cs
+++ b/Engine/Source/Programs/UnsyncUI/UnsyncUI/MainWindow.xaml.cs
@@ -227,7 +227,7 @@ namespace UnsyncUI
 				return;
 			}
 
-			UnsyncQueryUtil queryUtil = new UnsyncQueryUtil(Config.UnsyncPath, Config.RootProxy.Path);
+			UnsyncQueryUtil queryUtil = new UnsyncQueryUtil(Config.UnsyncPath, Config.RootProxy.Path, Config.RootProxy.Protocol);
 
 			try
 			{
@@ -345,7 +345,7 @@ namespace UnsyncUI
 					build.Config.DstPath,
 					build.Config.ScavengePath,
 					build.Config.DryRun, 
-					SelectedProxy?.Path,
+					SelectedProxy?.GetServerConfig(),
 					AdditionalArgs, 
 					build.Config.Exclusions, 
 					OnJobCompleted, 
diff --git a/Engine/Source/Programs/UnsyncUI/UnsyncUI/Query.cs b/Engine/Source/Programs/UnsyncUI/UnsyncUI/Query.cs
--- a/Engine/Source/Programs/UnsyncUI/UnsyncUI/Query.cs
+++ b/Engine/Source/Programs/UnsyncUI/UnsyncUI/Query.cs
@@ -12,10 +12,33 @@ using System.Collections;
 
 namespace UnsyncUI
 {
+	public class UnsyncServerConfig
+	{
+		public String address;
+		public String protocol;
+
+		public String GetCommandLineArgs()
+		{
+			List<String> args = new List<String>();
+
+			if (!string.IsNullOrWhiteSpace(address))
+			{
+				args.Add($"--server {address}");
+			}
+
+			if (!string.IsNullOrWhiteSpace(protocol))
+			{
+				args.Add($"--protocol {protocol}");
+			}
+
+			return string.Join(" ", args);
+		}
+	}
+
 	public class UnsyncQueryConfig
 	{
 		public String unsyncPath;
-		public String proxyAddress;
+		public UnsyncServerConfig server;
 	}
 
 	class SearchQueryResultEntry
@@ -63,11 +86,16 @@ namespace UnsyncUI
 			Config = InConfig;
 		}
 
-		public UnsyncQueryUtil(string unsyncPath, string proxyAddress)
+		public UnsyncQueryUtil(string unsyncPath, string serverAddress, string serverProtocol)
+			: this(unsyncPath, new UnsyncServerConfig { address = serverAddress, protocol = serverProtocol } )
+		{
+		}
+
+		public UnsyncQueryUtil(string unsyncPath, UnsyncServerConfig server)
 		{
 			Config = new UnsyncQueryConfig();
 			Config.unsyncPath = unsyncPath;
-			Config.proxyAddress = proxyAddress;
+			Config.server = server;
 		}
 
 		private string RunCommand(string argsStr)
@@ -106,13 +134,13 @@ namespace UnsyncUI
 
 		public List<UnsyncMirrorDesc> Mirrors()
 		{
-			string responseJson = RunCommand($"query mirrors --proxy {Config.proxyAddress}");
+			string responseJson = RunCommand($"query mirrors {Config.server.GetCommandLineArgs()}");
 			return JsonSerializer.Deserialize<List<UnsyncMirrorDesc>>(responseJson);
 		}
 
 		public LoginQueryResult Login()
 		{
-			string responseJson = RunCommand($"login --decode --proxy {Config.proxyAddress}");
+			string responseJson = RunCommand($"login --decode {Config.server.GetCommandLineArgs()}");
 			return JsonSerializer.Deserialize<LoginQueryResult>(responseJson);
 		}
 	}
@@ -197,7 +225,7 @@ namespace UnsyncUI
 		{
 			foreach (String query in queryStrings)
 			{
-				String argsStr = query + $" --proxy {Config.proxyAddress}";
+				String argsStr = query + $" {Config.server.GetCommandLineArgs()}";
 
 				var proc = new AsyncProcess(Config.unsyncPath, argsStr);
 				var responseJson = "";
