UnsyncUI - Only include the artifact id instead of full search string when syncing from Horde

diff --git a/Engine/Source/Programs/UnsyncUI/UnsyncUI/BuildSelectorControl.xaml.cs b/Engine/Source/Programs/UnsyncUI/UnsyncUI/BuildSelectorControl.xaml.cs
--- a/Engine/Source/Programs/UnsyncUI/UnsyncUI/BuildSelectorControl.xaml.cs
+++ b/Engine/Source/Programs/UnsyncUI/UnsyncUI/BuildSelectorControl.xaml.cs
@@ -163,18 +163,21 @@ namespace UnsyncUI
 							var files = await dirEnum.EnumerateFiles(path, cancellationToken.Token);
 							currentDir.ParseFileGroups(files.ToList(), template, out fileGroups);
 						}
+
+						string formattedPath = dirEnum.FormatArtifactPath(path);
+
 						if (fileGroups?.Any() == true)
 						{
 							foreach (var fileGroup in fileGroups)
 							{
-								Platforms.Add(new BuildPlatformModel(this, fileGroup.Platform, path, fileGroup.Flavor, fileGroup.Include));
+								Platforms.Add(new BuildPlatformModel(this, fileGroup.Platform, formattedPath, fileGroup.Flavor, fileGroup.Include));
 							}
-							Platforms.Add(new BuildPlatformModel(this, template.Platform, path, "All", null));
+							Platforms.Add(new BuildPlatformModel(this, template.Platform, formattedPath, "All", null));
 						}
 						else
 						{
 							// This folder is a valid build
-							Platforms.Add(new BuildPlatformModel(this, template.Platform, path, template.Flavor, null));
+							Platforms.Add(new BuildPlatformModel(this, template.Platform, formattedPath, template.Flavor, null));
 						}
 					}
 					else
diff --git a/Engine/Source/Programs/UnsyncUI/UnsyncUI/Query.cs b/Engine/Source/Programs/UnsyncUI/UnsyncUI/Query.cs
--- a/Engine/Source/Programs/UnsyncUI/UnsyncUI/Query.cs
+++ b/Engine/Source/Programs/UnsyncUI/UnsyncUI/Query.cs
@@ -181,6 +181,21 @@ namespace UnsyncUI
 			this.DirectorySchema = DirectorySchema;
 		}
 
+		public string FormatArtifactPath(string virtualPath) 
+		{
+			if (Config.server?.protocol == "horde")
+			{
+				// extract artifact ID assuming 'foobar#1234abcd' convention
+				int artifactIdPos = virtualPath.LastIndexOf('#');
+				string artifactId = artifactIdPos == -1 ? virtualPath : virtualPath.Substring(artifactIdPos);				
+				return artifactId;
+			}
+			else
+			{
+				return virtualPath; 
+			}
+		}
+
 		private async Task LazyInit(CancellationToken cancellationToken)
 		{
 			if (!Initialized)
diff --git a/Engine/Source/Programs/UnsyncUI/UnsyncUI/Utils.cs b/Engine/Source/Programs/UnsyncUI/UnsyncUI/Utils.cs
--- a/Engine/Source/Programs/UnsyncUI/UnsyncUI/Utils.cs
+++ b/Engine/Source/Programs/UnsyncUI/UnsyncUI/Utils.cs
@@ -73,6 +73,7 @@ namespace UnsyncUI
 	{
 		public Task<IEnumerable<string>> EnumerateDirectories(string path, CancellationToken token);
 		public Task<IEnumerable<string>> EnumerateFiles(string path, CancellationToken token);
+		public string FormatArtifactPath(string virtualPath) { return virtualPath; }
 	}
 
 	public class NativeDirectoryEnumerator : IDirectoryEnumerator
