UnsyncUI - Add basic template expansion for variables in the config file

* Allows sharing the same config for multiple projects that follow the same conventions
* Add support for setting ProjectDir, ProjectName and ProjectFile via command line using `--key value` syntax
* Automatically derive project name and directory from project file, unless explicitly specified
* Update version to 1.0.24

diff --git a/Engine/Source/Programs/UnsyncUI/UnsyncUI/App.xaml.cs b/Engine/Source/Programs/UnsyncUI/UnsyncUI/App.xaml.cs
--- a/Engine/Source/Programs/UnsyncUI/UnsyncUI/App.xaml.cs
+++ b/Engine/Source/Programs/UnsyncUI/UnsyncUI/App.xaml.cs
@@ -32,6 +32,8 @@ namespace UnsyncUI
 		internal bool EnableExperimentalFeatures = false;
 		internal bool EnableUserAuthentication = false;
 
+		internal Dictionary<string, string> Variables = new Dictionary<string, string>();
+
 		internal string ApplicationLog { get; private set; } = "";
 
 		internal void ClearApplicationLog()
@@ -76,6 +78,8 @@ namespace UnsyncUI
 		{
 			string configFile = e.Args.Length > 0 ? e.Args[0] : "unsyncui.xml";
 
+			List<string> VariableKeys = new List<string> { "ProjectFile", "ProjectName", "ProjectDir" };
+
 			for (int i=1; i<e.Args.Length; ++i)
 			{
 				bool isLast = i + 1 == e.Args.Length;
@@ -89,6 +93,42 @@ namespace UnsyncUI
 				{
 					EnableExperimentalFeatures = true;
 				}
+				else
+				{
+					string argLower = arg.ToLower();
+					foreach (string key in VariableKeys)
+					{
+						string keyLower = key.ToLower();
+						if ((argLower == $"--{keyLower}" || argLower == $"-{keyLower}") && !isLast)
+						{
+							string value = e.Args[i + 1];
+							Variables.Add(key, value);
+							++i;
+						}
+					}
+				}
+			}
+
+			// Derive project name from project file if one is not specified explicitly
+			if (!Variables.ContainsKey("ProjectName") && Variables.ContainsKey("ProjectFile"))
+			{
+				string ProjectFile = Variables["ProjectFile"];
+				string ProjectName = Path.GetFileNameWithoutExtension(ProjectFile);
+				if (!string.IsNullOrWhiteSpace(ProjectName))
+				{
+					Variables["ProjectName"] = ProjectName;
+				}
+			}
+
+			// Derive project directory from project file if one is not specified explicitly
+			if (!Variables.ContainsKey("ProjectDir") && Variables.ContainsKey("ProjectFile"))
+			{
+				string ProjectFile = Variables["ProjectFile"];
+				string ProjectDir = Path.GetDirectoryName(ProjectFile);
+				if (!string.IsNullOrWhiteSpace(ProjectDir))
+				{
+					Variables["ProjectDir"] = ProjectDir;
+				}
 			}
 
 			var oldWorkingDir = Environment.CurrentDirectory;
@@ -124,7 +164,7 @@ namespace UnsyncUI
 
 				try
 				{
-					Config = new Config(configFile, UnsyncPath);
+					Config = new Config(configFile, UnsyncPath, Variables);
 				}
 				catch (Exception ex)
 				{
diff --git a/Engine/Source/Programs/UnsyncUI/UnsyncUI/Config.cs b/Engine/Source/Programs/UnsyncUI/UnsyncUI/Config.cs
--- a/Engine/Source/Programs/UnsyncUI/UnsyncUI/Config.cs
+++ b/Engine/Source/Programs/UnsyncUI/UnsyncUI/Config.cs
@@ -222,10 +222,35 @@ namespace UnsyncUI
 
 		internal string loggedInUser;
 
-		public Config(string filename, string DefaultUnsyncPath = null)
+		internal void ApplyVariables(XElement rootNode, Dictionary<string, string> Variables)
+		{
+			foreach (XElement node in rootNode.Elements())
+			{
+				foreach (XAttribute attrib in node.Attributes())
+				{
+					if (attrib.Value.Contains("$("))
+					{
+						foreach (KeyValuePair<string, string> KV in Variables)
+						{
+							string resolved = attrib.Value.Replace($"$({KV.Key})", KV.Value);
+							attrib.SetValue(resolved);
+						}
+					}
+				}
+
+				ApplyVariables(node, Variables);
+			}
+		}
+
+		public Config(string filename, string DefaultUnsyncPath = null, Dictionary<string, string> Variables = null)
 		{
 			var rootNode = XDocument.Load(filename).Root;
 
+			if (Variables != null)
+			{
+				ApplyVariables(rootNode, Variables);
+			}
+
 			UnsyncPath = rootNode.Attribute("path")?.Value;
 			if (UnsyncPath == null)
 			{
diff --git a/Engine/Source/Programs/UnsyncUI/UnsyncUI/UnsyncUI.csproj b/Engine/Source/Programs/UnsyncUI/UnsyncUI/UnsyncUI.csproj
--- a/Engine/Source/Programs/UnsyncUI/UnsyncUI/UnsyncUI.csproj
+++ b/Engine/Source/Programs/UnsyncUI/UnsyncUI/UnsyncUI.csproj
@@ -8,7 +8,7 @@
 		<ApplicationManifest>app.manifest</ApplicationManifest>
         <IncludeAllContentForSelfExtract>true</IncludeAllContentForSelfExtract>
         <ApplicationIcon>unsync_48px.ico</ApplicationIcon>
-        <Version>$(VersionPrefix)1.0.23</Version>
+        <Version>$(VersionPrefix)1.0.24</Version>
 	</PropertyGroup>
 
 </Project>
