PR #12563: Add support for multiple audience values in decoded JWT

diff --git a/Engine/Source/Programs/UnsyncUI/UnsyncUI/Query.cs b/Engine/Source/Programs/UnsyncUI/UnsyncUI/Query.cs
--- a/Engine/Source/Programs/UnsyncUI/UnsyncUI/Query.cs
+++ b/Engine/Source/Programs/UnsyncUI/UnsyncUI/Query.cs
@@ -9,6 +9,7 @@ using System.Diagnostics;
 using System.Text.Json;
 using System.Linq;
 using System.Collections;
+using System.Text.Json.Serialization;
 
 namespace UnsyncUI
 {
@@ -59,7 +60,8 @@ namespace UnsyncUI
 	// Subset of decoded JWT claims
 	public class LoginQueryResult
 	{
-		public String aud { get; set; }
+		[JsonConverter(typeof(SingleOrArrayConverter))]
+		public List<String> aud { get; set; }
 		public String sub { get; set; }
 		public String iss { get; set; }
 		public List<String> scp { get; set; }
@@ -434,4 +436,28 @@ namespace UnsyncUI
 		}
 	}
 
+	public class SingleOrArrayConverter : JsonConverter<List<String>>
+	{
+		public override List<String> Read(ref Utf8JsonReader reader, Type typeToConvert, JsonSerializerOptions options)
+		{
+			return reader.TokenType switch
+			{
+				JsonTokenType.String => new List<String> { reader.GetString() },
+				JsonTokenType.StartArray => JsonSerializer.Deserialize<List<String>>(ref reader, options),
+				_ => throw new JsonException("Unexpected JSON format")
+			};
+		}
+
+		public override void Write(Utf8JsonWriter writer, List<String> value, JsonSerializerOptions options)
+		{
+			if (value.Count == 1)
+			{
+				writer.WriteStringValue(value[0]);
+			}
+			else
+			{
+				JsonSerializer.Serialize(writer, value, options);
+			}
+		}
+	}
 }
