UnsyncUI - Fix possible null object dereference when not using a proxy. Also include hidden directories when looking for manifests.

#12892
#12890

diff --git a/Engine/Source/Programs/UnsyncUI/UnsyncUI/MainWindow.xaml.cs b/Engine/Source/Programs/UnsyncUI/UnsyncUI/MainWindow.xaml.cs
--- a/Engine/Source/Programs/UnsyncUI/UnsyncUI/MainWindow.xaml.cs
+++ b/Engine/Source/Programs/UnsyncUI/UnsyncUI/MainWindow.xaml.cs
@@ -220,7 +220,7 @@ namespace UnsyncUI
 
 		private void LogIn()
 		{
-			if (Config.RootProxy.Path == null
+			if (Config.RootProxy?.Path == null
 				|| Config.UnsyncPath == null)
 			{
 				LoggedInUser = null;
diff --git a/Engine/Source/Programs/UnsyncUI/UnsyncUI/Utils.cs b/Engine/Source/Programs/UnsyncUI/UnsyncUI/Utils.cs
--- a/Engine/Source/Programs/UnsyncUI/UnsyncUI/Utils.cs
+++ b/Engine/Source/Programs/UnsyncUI/UnsyncUI/Utils.cs
@@ -101,6 +101,7 @@ namespace UnsyncUI
 			{
 				var enumOptions = new EnumerationOptions();
 				enumOptions.IgnoreInaccessible = true;
+				enumOptions.AttributesToSkip = 0; // include all files (hidden and system files are skipped by default)
 
 				using var cancel = token.Register(() => tcs.TrySetCanceled());
 				try
