UnsyncUI - Fix build list sorting by CL (should be using a numeric sort, rather than lexicographic)

diff --git a/Engine/Source/Programs/UnsyncUI/UnsyncUI/BuildSelectorControl.xaml.cs b/Engine/Source/Programs/UnsyncUI/UnsyncUI/BuildSelectorControl.xaml.cs
--- a/Engine/Source/Programs/UnsyncUI/UnsyncUI/BuildSelectorControl.xaml.cs
+++ b/Engine/Source/Programs/UnsyncUI/UnsyncUI/BuildSelectorControl.xaml.cs
@@ -115,7 +115,7 @@ namespace UnsyncUI
 	{
 		private CancellationTokenSource cts;
 		public string Stream => Template.Stream;
-		public string CL => Template.CL;
+		public UInt64 CL => Template.CL == null ? 0 : (UInt64)Template.CL;
 		public string Suffix => Template.Suffix == null ? "" : Template.Suffix;
 
 		public Config.Directory RootDir { get; }
@@ -257,7 +257,7 @@ namespace UnsyncUI
 
 			foreach (var term in searchTerms.Split(' '))
 			{
-				if (!build.CL.Contains(term, StringComparison.InvariantCultureIgnoreCase)
+				if (!build.CL.ToString().Contains(term, StringComparison.InvariantCultureIgnoreCase)
 				&& !build.Stream.Contains(term, StringComparison.InvariantCultureIgnoreCase)
 				&& !build.Suffix.Contains(term, StringComparison.InvariantCultureIgnoreCase))
 				{
@@ -542,7 +542,7 @@ namespace UnsyncUI
 		{
 			if (selectedBuild != null)
 			{
-				CopyStringToClipboard(selectedBuild.CL);
+				CopyStringToClipboard(selectedBuild.CL.ToString());
 			}
 		}
 
diff --git a/Engine/Source/Programs/UnsyncUI/UnsyncUI/Config.cs b/Engine/Source/Programs/UnsyncUI/UnsyncUI/Config.cs
--- a/Engine/Source/Programs/UnsyncUI/UnsyncUI/Config.cs
+++ b/Engine/Source/Programs/UnsyncUI/UnsyncUI/Config.cs
@@ -32,7 +32,7 @@ namespace UnsyncUI
 		internal struct BuildTemplate
 		{
 			public string Stream;
-			public string CL;
+			public UInt64? CL;
 			public string Suffix;
 			public string Platform;
 			public string Flavor;
@@ -95,7 +95,12 @@ namespace UnsyncUI
 				}
 				if (CL != null)
 				{
-					template.CL = RegexReplace(CL);
+					string CLString = RegexReplace(CL);
+					UInt64 CLNum = 0;
+					if (UInt64.TryParse(CLString, out CLNum))
+					{
+						template.CL = CLNum;
+					}
 				}
 				if (Suffix != null)
 				{
