UnsyncUI - Add timestamp column to build list

diff --git a/Engine/Source/Programs/UnsyncUI/UnsyncUI/BuildSelectorControl.xaml b/Engine/Source/Programs/UnsyncUI/UnsyncUI/BuildSelectorControl.xaml
--- a/Engine/Source/Programs/UnsyncUI/UnsyncUI/BuildSelectorControl.xaml
+++ b/Engine/Source/Programs/UnsyncUI/UnsyncUI/BuildSelectorControl.xaml
@@ -66,6 +66,7 @@
 							<GridViewColumn x:Name="colStream" Width="200" Header="Stream" DisplayMemberBinding="{Binding Stream}" />
 							<GridViewColumn x:Name="colCL" Width="100" Header="CL" DisplayMemberBinding="{Binding CL}" />
 							<GridViewColumn x:Name="colSuffix" Width="100" Header="Suffix" DisplayMemberBinding="{Binding Suffix}" />
+							<GridViewColumn x:Name="colTimestamp" Width="125" Header="Timestamp" DisplayMemberBinding="{Binding Timestamp}" />
 						</GridView>
 					</ListView.View>
 				</ListView>
diff --git a/Engine/Source/Programs/UnsyncUI/UnsyncUI/BuildSelectorControl.xaml.cs b/Engine/Source/Programs/UnsyncUI/UnsyncUI/BuildSelectorControl.xaml.cs
--- a/Engine/Source/Programs/UnsyncUI/UnsyncUI/BuildSelectorControl.xaml.cs
+++ b/Engine/Source/Programs/UnsyncUI/UnsyncUI/BuildSelectorControl.xaml.cs
@@ -118,8 +118,10 @@ namespace UnsyncUI
 		public UInt64 CL => Template.CL == null ? 0 : (UInt64)Template.CL;
 		public string Suffix => Template.Suffix == null ? "" : Template.Suffix;
 
+		public string Timestamp => DirEntry.MTime == DateTime.MinValue ? "" : DirEntry.MTime.ToString("yyyy-MM-dd HH:mm:ss");
+
 		public Config.Directory RootDir { get; }
-		public string Path { get; }
+		public string Path => DirEntry.Path;
 
 		internal Config.BuildTemplate Template { get; }
 
@@ -130,10 +132,12 @@ namespace UnsyncUI
 
 		private Config AppConfig;
 
-		internal BuildModel(string path, Config.Directory rootDir, Config.BuildTemplate template, Config config)
+		private DirectoryEntry DirEntry;
+
+		internal BuildModel(DirectoryEntry dirEntry, Config.Directory rootDir, Config.BuildTemplate template, Config config)
 		{
 			AppConfig = config;
-			Path = path;
+			DirEntry = dirEntry;
 			RootDir = rootDir;
 			Template = template;
 
@@ -155,12 +159,13 @@ namespace UnsyncUI
 			{
 				foreach (var childDir in await dirEnum.EnumerateDirectories(path, cancellationToken.Token))
 				{
-					if (System.IO.Path.GetFileName(childDir) == ".unsync")
+					if (System.IO.Path.GetFileName(childDir.Path) == ".unsync")
 					{
 						List<Config.BuildTemplate> fileGroups = null;
 						if (currentDir.FileGroups?.Any() == true)
 						{
-							var files = await dirEnum.EnumerateFiles(path, cancellationToken.Token);
+							var dirEntries = await dirEnum.EnumerateFiles(path, cancellationToken.Token);
+							var files = DirectoryEntry.ToPaths(dirEntries);
 							currentDir.ParseFileGroups(files.ToList(), template, out fileGroups);
 						}
 
@@ -184,7 +189,7 @@ namespace UnsyncUI
 					{
 						foreach (var subDir in currentDir.SubDirectories)
 						{
-							await EnumeratePlatforms(dirEnum, subDir, childDir, template, cancellationToken);
+							await EnumeratePlatforms(dirEnum, subDir, childDir.Path, template, cancellationToken);
 						}
 					}
 				}
@@ -565,11 +570,11 @@ namespace UnsyncUI
 
 		private void OpenDestinationLocation()
 		{
-			if (!Directory.Exists(FinalDstPath))
+			if (!System.IO.Directory.Exists(FinalDstPath))
 			{
 				try
 				{
-					Directory.CreateDirectory(FinalDstPath);
+					System.IO.Directory.CreateDirectory(FinalDstPath);
 				}
 				catch (Exception)
 				{
diff --git a/Engine/Source/Programs/UnsyncUI/UnsyncUI/Config.cs b/Engine/Source/Programs/UnsyncUI/UnsyncUI/Config.cs
--- a/Engine/Source/Programs/UnsyncUI/UnsyncUI/Config.cs
+++ b/Engine/Source/Programs/UnsyncUI/UnsyncUI/Config.cs
@@ -164,21 +164,21 @@ namespace UnsyncUI
 			public List<string> Exclusions { get; set; }
 			public List<Directory> Children { get; set; }
 
-			private Task EnumerateBuilds(IDirectoryEnumerator dirEnum, ITargetBlock<BuildModel> pipe, Directory d, string path, BuildTemplate template, CancellationToken cancelToken)
+			private Task EnumerateBuilds(IDirectoryEnumerator dirEnum, ITargetBlock<BuildModel> pipe, Directory d, DirectoryEntry dirEntry, BuildTemplate template, CancellationToken cancelToken)
 			{
 				return Task.Run(async () =>
 				{
-					if (!d.Parse(path, ref template))
+					if (!d.Parse(dirEntry.Path, ref template))
 						return;
 
 					if (template.IsStreamClFound)
 					{
-						pipe.Post(new BuildModel(path, d, template, AppConfig));
+						pipe.Post(new BuildModel(dirEntry, d, template, AppConfig));
 					}
 					else
 					{
 						var jobs = new List<Task>();
-						foreach (var childDir in await dirEnum.EnumerateDirectories(path, cancelToken))
+						foreach (var childDir in await dirEnum.EnumerateDirectories(dirEntry.Path, cancelToken))
 						{
 							jobs.AddRange(d.SubDirectories.Select(s => EnumerateBuilds(dirEnum, pipe, s, childDir, template, cancelToken)));
 						}
@@ -425,9 +425,15 @@ namespace UnsyncUI
 			}
 		}
 
+		public BuildModel CreateBuildModel(string path, DateTime timestamp, Config.Directory rootDir)
+		{
+			var dirEntry = new DirectoryEntry { Path = path, MTime = timestamp };
+			return new BuildModel(dirEntry, rootDir, default(BuildTemplate), this);
+		}
+
 		public BuildModel CreateBuildModel(string path, Config.Directory rootDir)
 		{
-			return new BuildModel(path, rootDir, default(BuildTemplate), this);
+			return CreateBuildModel(path, DateTime.MinValue, rootDir);
 		}
 
 	} // class Config
diff --git a/Engine/Source/Programs/UnsyncUI/UnsyncUI/MainWindow.xaml b/Engine/Source/Programs/UnsyncUI/UnsyncUI/MainWindow.xaml
--- a/Engine/Source/Programs/UnsyncUI/UnsyncUI/MainWindow.xaml
+++ b/Engine/Source/Programs/UnsyncUI/UnsyncUI/MainWindow.xaml
@@ -73,7 +73,7 @@
 	</Window.Resources>
 	<Grid Margin="4">
 		<Grid.ColumnDefinitions>
-			<ColumnDefinition Width="476"/>
+			<ColumnDefinition Width="588"/>
 			<ColumnDefinition Width="8"/>
 			<ColumnDefinition Width="*"/>
 		</Grid.ColumnDefinitions>
diff --git a/Engine/Source/Programs/UnsyncUI/UnsyncUI/Query.cs b/Engine/Source/Programs/UnsyncUI/UnsyncUI/Query.cs
--- a/Engine/Source/Programs/UnsyncUI/UnsyncUI/Query.cs
+++ b/Engine/Source/Programs/UnsyncUI/UnsyncUI/Query.cs
@@ -377,16 +377,16 @@ namespace UnsyncUI
 			}
 		}
 
-		private Task<IEnumerable<string>> EnumerateEntries(string path, CancellationToken token, EntryFilter filter)
+		private Task<IEnumerable<DirectoryEntry>> EnumerateEntries(string path, CancellationToken token, EntryFilter filter)
 		{
-			var tcs = new TaskCompletionSource<IEnumerable<string>>();
+			var tcs = new TaskCompletionSource<IEnumerable<DirectoryEntry>>();
 			Task.Run(async () =>
 			{
 				using var cancel = token.Register(() => tcs.TrySetCanceled());
 				try
 				{
 					await LazyInit(token);
-					var dirs = new List<string>();
+					var dirs = new List<DirectoryEntry>();
 
 					int pathDepth = GetDirectoryDepth(path);
 
@@ -406,7 +406,8 @@ namespace UnsyncUI
 							&& entry.FullPath.StartsWith(requiredPrefix))
 						{
 							//String relativePath = Path.GetRelativePath(path, entry.FullPath);
-							dirs.Add(entry.FullPath);
+							DirectoryEntry dirEntry = new DirectoryEntry { Path = entry.FullPath, MTime = DateTime.FromFileTime((long)entry.MTime) };
+							dirs.Add(dirEntry);
 						}
 					}
 
@@ -420,14 +421,14 @@ namespace UnsyncUI
 			return tcs.Task;
 		}
 
-		public Task<IEnumerable<string>> EnumerateDirectories(string path, CancellationToken token)
+		public Task<IEnumerable<DirectoryEntry>> EnumerateDirectories(string path, CancellationToken token)
 		{
 			EntryFilter filter;
 			filter.IncludeDirectories = true;
 			filter.IncludeFiles = false;
 			return EnumerateEntries(path, token, filter);
 		}
-		public Task<IEnumerable<string>> EnumerateFiles(string path, CancellationToken token)
+		public Task<IEnumerable<DirectoryEntry>> EnumerateFiles(string path, CancellationToken token)
 		{
 			EntryFilter filter;
 			filter.IncludeDirectories = false;
diff --git a/Engine/Source/Programs/UnsyncUI/UnsyncUI/UnsyncUI.csproj b/Engine/Source/Programs/UnsyncUI/UnsyncUI/UnsyncUI.csproj
--- a/Engine/Source/Programs/UnsyncUI/UnsyncUI/UnsyncUI.csproj
+++ b/Engine/Source/Programs/UnsyncUI/UnsyncUI/UnsyncUI.csproj
@@ -8,7 +8,7 @@
 		<ApplicationManifest>app.manifest</ApplicationManifest>
         <IncludeAllContentForSelfExtract>true</IncludeAllContentForSelfExtract>
         <ApplicationIcon>unsync_48px.ico</ApplicationIcon>
-        <Version>$(VersionPrefix)1.0.24</Version>
+        <Version>$(VersionPrefix)1.0.25</Version>
 	</PropertyGroup>
 
 </Project>
diff --git a/Engine/Source/Programs/UnsyncUI/UnsyncUI/Utils.cs b/Engine/Source/Programs/UnsyncUI/UnsyncUI/Utils.cs
--- a/Engine/Source/Programs/UnsyncUI/UnsyncUI/Utils.cs
+++ b/Engine/Source/Programs/UnsyncUI/UnsyncUI/Utils.cs
@@ -69,20 +69,32 @@ namespace UnsyncUI
         { }
     }
 
+	public struct DirectoryEntry
+	{
+		public string Path;
+		public DateTime MTime;
+
+		public static IEnumerable<string> ToPaths(IEnumerable<DirectoryEntry> dirEntries)
+		{
+			return dirEntries.Select(entry => entry.Path);
+		}
+	}
+
 	public interface IDirectoryEnumerator
 	{
-		public Task<IEnumerable<string>> EnumerateDirectories(string path, CancellationToken token);
-		public Task<IEnumerable<string>> EnumerateFiles(string path, CancellationToken token);
+		public Task<IEnumerable<DirectoryEntry>> EnumerateDirectories(string path, CancellationToken token);
+		public Task<IEnumerable<DirectoryEntry>> EnumerateFiles(string path, CancellationToken token);
 		public string FormatArtifactPath(string virtualPath) { return virtualPath; }
 	}
 
+
 	public class NativeDirectoryEnumerator : IDirectoryEnumerator
 	{
-		public Task<IEnumerable<string>> EnumerateDirectories(string path, CancellationToken token)
+		public Task<IEnumerable<DirectoryEntry>> EnumerateDirectories(string path, CancellationToken token)
 		{
 			return AsyncIO.EnumerateDirectoriesAsync(path, token);
 		}
-		public Task<IEnumerable<string>> EnumerateFiles(string path, CancellationToken token)
+		public Task<IEnumerable<DirectoryEntry>> EnumerateFiles(string path, CancellationToken token)
 		{
 			return AsyncIO.EnumerateFilesAsync(path, token);
 		}
@@ -92,13 +104,15 @@ namespace UnsyncUI
     {
 		private static SemaphoreSlim mutex = new SemaphoreSlim(4);
 
-        public static Task<IEnumerable<string>> EnumerateDirectoriesAsync(string path, CancellationToken token)
+        public static Task<IEnumerable<DirectoryEntry>> EnumerateDirectoriesAsync(string path, CancellationToken token)
         {
 			App.Current?.LogMessage($"Enumerating: {path}");
 
-			var tcs = new TaskCompletionSource<IEnumerable<string>>();
+			var tcs = new TaskCompletionSource<IEnumerable<DirectoryEntry>>();
             Task.Run(() =>
 			{
+				var dirInfo = new DirectoryInfo(path);
+
 				var enumOptions = new EnumerationOptions();
 				enumOptions.IgnoreInaccessible = true;
 				enumOptions.AttributesToSkip = 0; // include all files (hidden and system files are skipped by default)
@@ -109,8 +123,9 @@ namespace UnsyncUI
 					mutex.Wait(token);
 					try
 					{
-						var dirs = Directory.EnumerateDirectories(path, "*", enumOptions).ToList();
-						tcs.TrySetResult(dirs);
+						var dirs = dirInfo.EnumerateDirectories("*", enumOptions);
+						var dirEntries = dirs.Select(d => new DirectoryEntry { Path = d.FullName, MTime = d.LastWriteTime});
+						tcs.TrySetResult(dirEntries);
 					}
 					catch (Exception ex)
 					{
@@ -128,9 +143,9 @@ namespace UnsyncUI
 
             return tcs.Task;
 		}
-        public static Task<IEnumerable<string>> EnumerateFilesAsync(string path, CancellationToken token)
+        public static Task<IEnumerable<DirectoryEntry>> EnumerateFilesAsync(string path, CancellationToken token)
         {
-			var tcs = new TaskCompletionSource<IEnumerable<string>>();
+			var tcs = new TaskCompletionSource<IEnumerable<DirectoryEntry>>();
             Task.Run(() =>
 			{
 				using var cancel = token.Register(() => tcs.TrySetCanceled());
@@ -142,12 +157,19 @@ namespace UnsyncUI
 						var timer = new Stopwatch();
 						timer.Start();
 
-						var dirs = Directory.EnumerateFiles(path).ToList();
+						var dirInfo = new DirectoryInfo(path);
+
+						var enumOptions = new EnumerationOptions();
+						enumOptions.IgnoreInaccessible = true;
+						enumOptions.AttributesToSkip = 0; // include all files (hidden and system files are skipped by default)
+
+						var dirs = dirInfo.EnumerateDirectories("*", enumOptions);
+						var dirEntries = dirs.Select(d => new DirectoryEntry { Path = d.FullName, MTime = d.LastWriteTime });
 
 						timer.Stop();
 						App.Current?.LogMessage($"Time: {timer.Elapsed.TotalSeconds:0.000} s - {path}");
 
-						tcs.TrySetResult(dirs);
+						tcs.TrySetResult(dirEntries);
 					}
 					catch (Exception ex)
 					{
