UnsyncUI - Fix incorrect usage of EnumerateDirectories when EnumerateFiles is expected

Also fix a few minor bugs:
- Propagate ReadStdErr argument from RunAndGetOutput into RunAsync
- Fix possible overflow when converting FileTime stored as UInt64 to DateTime

All issues identified by BugHawk.

diff --git a/Engine/Source/Programs/UnsyncUI/UnsyncUI/Query.cs b/Engine/Source/Programs/UnsyncUI/UnsyncUI/Query.cs
--- a/Engine/Source/Programs/UnsyncUI/UnsyncUI/Query.cs
+++ b/Engine/Source/Programs/UnsyncUI/UnsyncUI/Query.cs
@@ -406,7 +406,12 @@ namespace UnsyncUI
 							&& entry.FullPath.StartsWith(requiredPrefix))
 						{
 							//String relativePath = Path.GetRelativePath(path, entry.FullPath);
-							DirectoryEntry dirEntry = new DirectoryEntry { Path = entry.FullPath, MTime = DateTime.FromFileTime((long)entry.MTime) };
+							DateTime mtime = DateTime.FromFileTime((long)Math.Min(entry.MTime, (UInt64)DateTime.MaxValue.ToFileTime()));
+							DirectoryEntry dirEntry = new DirectoryEntry 
+							{ 
+								Path = entry.FullPath,
+								MTime = mtime,
+							};
 							dirs.Add(dirEntry);
 						}
 					}
diff --git a/Engine/Source/Programs/UnsyncUI/UnsyncUI/Utils.cs b/Engine/Source/Programs/UnsyncUI/UnsyncUI/Utils.cs
--- a/Engine/Source/Programs/UnsyncUI/UnsyncUI/Utils.cs
+++ b/Engine/Source/Programs/UnsyncUI/UnsyncUI/Utils.cs
@@ -163,7 +163,7 @@ namespace UnsyncUI
 						enumOptions.IgnoreInaccessible = true;
 						enumOptions.AttributesToSkip = 0; // include all files (hidden and system files are skipped by default)
 
-						var dirs = dirInfo.EnumerateDirectories("*", enumOptions);
+						var dirs = dirInfo.EnumerateFiles("*", enumOptions);
 						var dirEntries = dirs.Select(d => new DirectoryEntry { Path = d.FullName, MTime = d.LastWriteTime });
 
 						timer.Stop();
@@ -394,7 +394,7 @@ namespace UnsyncUI
 		{
 			var result = "";
 			var cancellationToken = new CancellationToken();
-			await foreach (var str in RunAsync(cancellationToken, false /*ReadStdErr*/))
+			await foreach (var str in RunAsync(cancellationToken, ReadStdErr))
 			{
 				result += str;
 			}
