Update UnsyncUI to net8.0

#rnx
#Synced-CL 49686086

diff --git a/Engine/Source/Programs/UnsyncUI/UnsyncUI/UnsyncUI.csproj b/Engine/Source/Programs/UnsyncUI/UnsyncUI/UnsyncUI.csproj
--- a/Engine/Source/Programs/UnsyncUI/UnsyncUI/UnsyncUI.csproj
+++ b/Engine/Source/Programs/UnsyncUI/UnsyncUI/UnsyncUI.csproj
@@ -2,13 +2,13 @@
 
 	<PropertyGroup>
 		<OutputType>WinExe</OutputType>
-		<TargetFramework>net6.0-windows</TargetFramework>
+		<TargetFramework>net8.0-windows</TargetFramework>
 		<UseWPF>true</UseWPF>
 		<UseWindowsForms>true</UseWindowsForms>
 		<ApplicationManifest>app.manifest</ApplicationManifest>
-        <IncludeAllContentForSelfExtract>true</IncludeAllContentForSelfExtract>
-        <ApplicationIcon>unsync_48px.ico</ApplicationIcon>
-        <Version>$(VersionPrefix)1.0.25</Version>
+		<IncludeAllContentForSelfExtract>true</IncludeAllContentForSelfExtract>
+		<ApplicationIcon>unsync_48px.ico</ApplicationIcon>
+		<Version>$(VersionPrefix)1.0.25</Version>
 	</PropertyGroup>
 
 </Project>
