UnsyncUI - Show server announcements, fix garbled output for warnings/errors in the sync log window

* Also update publishing config to use .net 8.0

#Synced-CL 50943321

diff --git a/Engine/Source/Programs/UnsyncUI/UnsyncUI/App.xaml.cs b/Engine/Source/Programs/UnsyncUI/UnsyncUI/App.xaml.cs
--- a/Engine/Source/Programs/UnsyncUI/UnsyncUI/App.xaml.cs
+++ b/Engine/Source/Programs/UnsyncUI/UnsyncUI/App.xaml.cs
@@ -31,6 +31,7 @@ namespace UnsyncUI
 
 		internal bool EnableExperimentalFeatures = false;
 		internal bool EnableUserAuthentication = false;
+		internal bool UnsyncSupportsStdout = false;
 
 		internal Dictionary<string, string> Variables = new Dictionary<string, string>();
 
@@ -201,6 +202,34 @@ namespace UnsyncUI
 				Config.EnableUserAuthentication = EnableUserAuthentication;
 			}
 
+			if (UnsyncPath != null)
+			{
+				try
+				{
+					using (var proc = new System.Diagnostics.Process())
+					{
+						proc.StartInfo = new System.Diagnostics.ProcessStartInfo(UnsyncPath, "sync --help")
+						{
+							CreateNoWindow = true,
+							UseShellExecute = false,
+							RedirectStandardOutput = true,
+							RedirectStandardError = true
+						};
+						proc.Start();
+						string helpText = proc.StandardOutput.ReadToEnd();
+						if (proc.WaitForExit(5000))
+						{
+							UnsyncSupportsStdout = helpText.Contains("--stdout");
+						}
+						else
+						{
+							proc.Kill();
+						}
+					}
+				}
+				catch { }
+			}
+
 			base.OnStartup(e);
 		}
 
diff --git a/Engine/Source/Programs/UnsyncUI/UnsyncUI/JobControl.xaml.cs b/Engine/Source/Programs/UnsyncUI/UnsyncUI/JobControl.xaml.cs
--- a/Engine/Source/Programs/UnsyncUI/UnsyncUI/JobControl.xaml.cs
+++ b/Engine/Source/Programs/UnsyncUI/UnsyncUI/JobControl.xaml.cs
@@ -259,7 +259,7 @@ namespace UnsyncUI
 			cts.Cancel();
 		}
 
-		public bool IsDuplicate(JobModel otherJob) 
+		public bool IsDuplicate(JobModel otherJob)
 		{
 			return (Build.Build.Path == otherJob.Build.Build.Path)
 				&& (Build.Platform == otherJob.Build.Platform)
@@ -361,6 +361,10 @@ namespace UnsyncUI
 
 				// Launch unsync.exe
 				var args = new List<string>() { "-v", "--progress" };
+				if (App.Current.UnsyncSupportsStdout)
+				{
+					args.Add("--stdout");
+				}
 				if (DryRun)
 				{
 					args.Add("-d");
@@ -416,8 +420,8 @@ namespace UnsyncUI
 							command = true;
 						}
 
-						var builder = command 
-							? commandBuilder 
+						var builder = command
+							? commandBuilder
 							: lineBuilder;
 
 						switch (current)
diff --git a/Engine/Source/Programs/UnsyncUI/UnsyncUI/MainWindow.xaml b/Engine/Source/Programs/UnsyncUI/UnsyncUI/MainWindow.xaml
--- a/Engine/Source/Programs/UnsyncUI/UnsyncUI/MainWindow.xaml
+++ b/Engine/Source/Programs/UnsyncUI/UnsyncUI/MainWindow.xaml
@@ -18,6 +18,8 @@
 	</Window.TaskbarItemInfo>
 	<Window.Resources>
 
+		<local:ConsoleColorConverter x:Key="ConsoleColorConverter"/>
+
 		<local:MainTabTemplateSelector x:Key="mainTabSelector">
 			<local:MainTabTemplateSelector.ProjectTemplate>
 				<DataTemplate>
@@ -71,7 +73,37 @@
 		</local:MainTabTemplateSelector>
 
 	</Window.Resources>
-	<Grid Margin="4">
+	<Grid>
+		<Grid.RowDefinitions>
+			<RowDefinition Height="auto"/>
+			<RowDefinition Height="*"/>
+		</Grid.RowDefinitions>
+
+		<!-- Announcement Banner — hidden when no announcements -->
+		<Border Grid.Row="0"
+				Visibility="{Binding AnnouncementBanner.IsVisible, Converter={StaticResource VisConv_Collapsed}}"
+				BorderBrush="LightGray" BorderThickness="0,0,0,1" Padding="8,4" Margin="4,4,4,0">
+			<Grid>
+				<Grid.ColumnDefinitions>
+					<ColumnDefinition Width="*"/>
+					<ColumnDefinition Width="auto"/>
+				</Grid.ColumnDefinitions>
+				<ItemsControl Grid.Column="0" ItemsSource="{Binding AnnouncementBanner.Announcements}">
+					<ItemsControl.ItemTemplate>
+						<DataTemplate>
+							<TextBlock Text="{Binding Message}" TextWrapping="Wrap"
+									   Foreground="{Binding Color, Converter={StaticResource ConsoleColorConverter}}"
+									   Margin="0,2"/>
+						</DataTemplate>
+					</ItemsControl.ItemTemplate>
+				</ItemsControl>
+				<TextBlock Grid.Column="1" VerticalAlignment="Top" Margin="8,2,0,0">
+					<Hyperlink Command="{Binding AnnouncementBanner.OnDismissClicked}">Dismiss</Hyperlink>
+				</TextBlock>
+			</Grid>
+		</Border>
+
+		<Grid Grid.Row="1" Margin="4">
 		<Grid.ColumnDefinitions>
 			<ColumnDefinition Width="588"/>
 			<ColumnDefinition Width="8"/>
@@ -268,5 +300,6 @@
 		</Grid>
 
 	</Grid>
+	</Grid>
 
 </Window>
diff --git a/Engine/Source/Programs/UnsyncUI/UnsyncUI/MainWindow.xaml.cs b/Engine/Source/Programs/UnsyncUI/UnsyncUI/MainWindow.xaml.cs
--- a/Engine/Source/Programs/UnsyncUI/UnsyncUI/MainWindow.xaml.cs
+++ b/Engine/Source/Programs/UnsyncUI/UnsyncUI/MainWindow.xaml.cs
@@ -70,6 +70,62 @@ namespace UnsyncUI
 				: CustomTemplate;
 	}
 
+	public sealed class AnnouncementItem : BaseModel
+	{
+		public string Message { get; }
+		public ConsoleColor Color { get; }
+
+		public AnnouncementItem(string message, string level)
+		{
+			Message = (message ?? "").Replace("\r\n", " ").Replace("\n", " ").Replace("\r", " ");
+			Color = level switch
+			{
+				"error" => ConsoleColor.Red,
+				"warning" => ConsoleColor.DarkYellow,
+				_ => ConsoleColor.Black,
+			};
+		}
+	}
+
+	public sealed class AnnouncementBannerModel : BaseModel
+	{
+		private HashSet<string> seenIds = new HashSet<string>();
+
+		public ObservableCollection<AnnouncementItem> Announcements { get; } = new ObservableCollection<AnnouncementItem>();
+
+		private bool isVisible = false;
+		public bool IsVisible
+		{
+			get => isVisible;
+			set => SetProperty(ref isVisible, value);
+		}
+
+		public Command OnDismissClicked { get; }
+
+		public AnnouncementBannerModel()
+		{
+			OnDismissClicked = new Command(() => IsVisible = false) { Enabled = true };
+		}
+
+		public void AddAnnouncements(IEnumerable<HelloAnnouncement> announcements)
+		{
+			foreach (var a in announcements)
+			{
+				if (a.id != null && !seenIds.Add(a.id))
+				{
+					continue;
+				}
+
+				Announcements.Add(new AnnouncementItem(a.message, a.level));
+			}
+
+			if (Announcements.Count > 0)
+			{
+				IsVisible = true;
+			}
+		}
+	}
+
 	public abstract class TabModel : BaseModel
 	{
 		public virtual void OnSelected() { }
@@ -154,7 +210,9 @@ namespace UnsyncUI
 
 	public sealed class MainWindowModel : BaseModel
 	{
-		private Config.Proxy selectedProxy = 
+		public AnnouncementBannerModel AnnouncementBanner { get; } = new AnnouncementBannerModel();
+
+		private Config.Proxy selectedProxy =
 			App.Current.Config?.Proxies?.Where(p => p.Name.Equals(App.Current.UserConfig.Proxy, StringComparison.InvariantCultureIgnoreCase))?.FirstOrDefault() ??
 			App.Current.Config?.Proxies?.Where(p => p.Path != null).FirstOrDefault();
 
@@ -317,6 +375,37 @@ namespace UnsyncUI
 			Tabs.Add(new CustomModel(OnBuildsSelected));
 
 			SelectedTab = Tabs.FirstOrDefault();
+
+			FetchAnnouncements();
+		}
+
+		private void FetchAnnouncements()
+		{
+			if (Config?.RootProxy?.Path == null || Config?.UnsyncPath == null)
+			{
+				return;
+			}
+
+			System.Threading.Tasks.Task.Run(() =>
+			{
+				try
+				{
+					var queryUtil = new UnsyncQueryUtil(Config.UnsyncPath, Config.RootProxy.GetServerConfig());
+					var announcements = queryUtil.Hello();
+
+					if (announcements.Count > 0)
+					{
+						App.Current.Dispatcher.InvokeAsync(() =>
+						{
+							AnnouncementBanner.AddAnnouncements(announcements);
+						});
+					}
+				}
+				catch (Exception ex)
+				{
+					App.Current.LogDebug($"Failed to fetch announcements: {ex.Message}");
+				}
+			});
 		}
 
 		private void ClearQueue()
diff --git a/Engine/Source/Programs/UnsyncUI/UnsyncUI/Query.cs b/Engine/Source/Programs/UnsyncUI/UnsyncUI/Query.cs
--- a/Engine/Source/Programs/UnsyncUI/UnsyncUI/Query.cs
+++ b/Engine/Source/Programs/UnsyncUI/UnsyncUI/Query.cs
@@ -57,6 +57,18 @@ namespace UnsyncUI
 		public List<SearchQueryResultEntry> entries { get; set; }
 	}
 
+	public class HelloAnnouncement
+	{
+		public string id { get; set; }
+		public string level { get; set; }
+		public string message { get; set; }
+	}
+
+	class HelloResponse
+	{
+		public List<HelloAnnouncement> announcements { get; set; }
+	}
+
 	// Subset of decoded JWT claims
 	public class LoginQueryResult
 	{
@@ -145,6 +157,26 @@ namespace UnsyncUI
 			string responseJson = RunCommand($"login --decode {Config.server.GetCommandLineArgs()}");
 			return JsonSerializer.Deserialize<LoginQueryResult>(responseJson);
 		}
+
+		public List<HelloAnnouncement> Hello()
+		{
+			string serverAddress = Config.server?.address;
+			if (string.IsNullOrWhiteSpace(serverAddress))
+			{
+				return new List<HelloAnnouncement>();
+			}
+
+			string scheme = serverAddress.EndsWith(":443") || serverAddress.Contains(":443/") ? "https" : "http";
+			string uriString = serverAddress.Contains("://") ? serverAddress : $"{scheme}://{serverAddress}";
+			if (!Uri.TryCreate(uriString, UriKind.Absolute, out Uri baseUri))
+			{
+				return new List<HelloAnnouncement>();
+			}
+			string url = new Uri(baseUri, "/api/v1/hello").ToString();
+			string responseJson = RunCommand($"query http-get {url}");
+			HelloResponse response = JsonSerializer.Deserialize<HelloResponse>(responseJson);
+			return response?.announcements ?? new List<HelloAnnouncement>();
+		}
 	}
 
 	
diff --git a/Engine/Source/Programs/UnsyncUI/UnsyncUI/UnsyncUI.csproj b/Engine/Source/Programs/UnsyncUI/UnsyncUI/UnsyncUI.csproj
--- a/Engine/Source/Programs/UnsyncUI/UnsyncUI/UnsyncUI.csproj
+++ b/Engine/Source/Programs/UnsyncUI/UnsyncUI/UnsyncUI.csproj
@@ -8,7 +8,7 @@
 		<ApplicationManifest>app.manifest</ApplicationManifest>
 		<IncludeAllContentForSelfExtract>true</IncludeAllContentForSelfExtract>
 		<ApplicationIcon>unsync_48px.ico</ApplicationIcon>
-		<Version>$(VersionPrefix)1.0.25</Version>
+		<Version>$(VersionPrefix)1.0.26</Version>
 	</PropertyGroup>
 
 </Project>
