UnsyncUI - Minor UX polish

* Auto-follow the log
* Allow fine grained log text selection

#Synced-CL 51327116

diff --git a/Engine/Source/Programs/UnsyncUI/UnsyncUI/JobControl.xaml b/Engine/Source/Programs/UnsyncUI/UnsyncUI/JobControl.xaml
--- a/Engine/Source/Programs/UnsyncUI/UnsyncUI/JobControl.xaml
+++ b/Engine/Source/Programs/UnsyncUI/UnsyncUI/JobControl.xaml
@@ -3,10 +3,10 @@
 <UserControl x:Class="UnsyncUI.JobControl"
              xmlns="http://schemas.microsoft.com/winfx/2006/xaml/presentation"
              xmlns:x="http://schemas.microsoft.com/winfx/2006/xaml"
-             xmlns:mc="http://schemas.openxmlformats.org/markup-compatibility/2006" 
-             xmlns:d="http://schemas.microsoft.com/expression/blend/2008" 
+             xmlns:mc="http://schemas.openxmlformats.org/markup-compatibility/2006"
+             xmlns:d="http://schemas.microsoft.com/expression/blend/2008"
              xmlns:local="clr-namespace:UnsyncUI"
-             mc:Ignorable="d" 
+             mc:Ignorable="d"
              d:DesignWidth="800">
 	<UserControl.Resources>
 		<local:ConsoleColorConverter x:Key="ConsoleColorConverter"/>
@@ -42,7 +42,7 @@
             </TextBlock>
 
 			<ProgressBar Grid.Row="1" Grid.Column="0" Grid.ColumnSpan="5"
-                         Value="{Binding TotalProgress, Mode=OneWay}" 
+                         Value="{Binding TotalProgress, Mode=OneWay}"
                          IsIndeterminate="{Binding IsIndeterminate, Mode=OneWay}"
                          Visibility="{Binding ShowProgressBar, Converter={StaticResource VisConv_Collapsed}}" />
 		</Grid>
@@ -67,41 +67,44 @@
                     <TextBlock Text="{Binding Build.FullPath, Mode=OneWay}"/>
                 </Hyperlink>
             </TextBlock>
-			
+
 			<TextBlock Grid.Column="2" Grid.Row="1">
                 <Hyperlink Command="{Binding OnDestPathClicked}">
                     <TextBlock Text="{Binding DstPath, Mode=OneWay}"/>
                 </Hyperlink>
             </TextBlock>
-			
+
 			<TextBlock Grid.Column="2" Grid.Row="2" Text="{Binding ProxyStr}" />
 
 			<GroupBox Grid.Row="3" Grid.ColumnSpan="3" Header="Log" Visibility="{Binding LogLines.Count, Converter={StaticResource VisConv_Collapsed}}">
-				<GroupBox.ContextMenu>
-					<ContextMenu>
-						<MenuItem Header="Copy to Clipboard" Command="{Binding CopyLogClicked}"/>
-						<MenuItem Header="Save As..." Command="{Binding SaveLogClicked}"/>
-					</ContextMenu>
-				</GroupBox.ContextMenu>
-				<ItemsControl ItemsSource="{Binding LogLines}">
-					<ItemsControl.ItemTemplate>
-						<DataTemplate>
-							<TextBlock Text="{Binding Text}" FontFamily="Consolas" Foreground="{Binding Color, Converter={StaticResource ConsoleColorConverter}}"/>
-						</DataTemplate>
-					</ItemsControl.ItemTemplate>
-					<ItemsControl.ItemsPanel>
-						<ItemsPanelTemplate>
-							<VirtualizingStackPanel IsVirtualizing="True" VirtualizationMode="Recycling" />
-						</ItemsPanelTemplate>
-					</ItemsControl.ItemsPanel>
-					<ItemsControl.Template>
-						<ControlTemplate TargetType="ItemsControl">
-							<ScrollViewer CanContentScroll="True" Focusable="True" MaxHeight="{Binding MaxLogScrollViewHeight}">
-								<ItemsPresenter />
-							</ScrollViewer>
-						</ControlTemplate>
-					</ItemsControl.Template>
-				</ItemsControl>
+				<Grid>
+					<RichTextBox x:Name="LogRichTextBox"
+								 IsReadOnly="True"
+								 IsReadOnlyCaretVisible="True"
+								 BorderThickness="0"
+								 FontFamily="Consolas"
+								 VerticalScrollBarVisibility="Auto"
+								 HorizontalScrollBarVisibility="Disabled"
+								 MaxHeight="{Binding MaxLogScrollViewHeight}"
+								 Margin="0">
+						<RichTextBox.ContextMenu>
+							<ContextMenu>
+								<MenuItem Header="Copy" Command="ApplicationCommands.Copy"/>
+								<MenuItem Header="Copy All" Command="{Binding CopyLogClicked}"/>
+								<MenuItem Header="Save As..." Command="{Binding SaveLogClicked}"/>
+							</ContextMenu>
+						</RichTextBox.ContextMenu>
+						<RichTextBox.Resources>
+							<Style TargetType="Paragraph">
+								<Setter Property="Margin" Value="0"/>
+							</Style>
+						</RichTextBox.Resources>
+						<FlowDocument PagePadding="0"/>
+					</RichTextBox>
+					<TextBlock x:Name="FollowLink" HorizontalAlignment="Right" VerticalAlignment="Bottom" Visibility="Collapsed" Margin="0,0,28,4">
+						<Hyperlink Click="FollowLink_Click">Follow</Hyperlink>
+					</TextBlock>
+				</Grid>
 			</GroupBox>
 		</Grid>
 	</Grid>
diff --git a/Engine/Source/Programs/UnsyncUI/UnsyncUI/JobControl.xaml.cs b/Engine/Source/Programs/UnsyncUI/UnsyncUI/JobControl.xaml.cs
--- a/Engine/Source/Programs/UnsyncUI/UnsyncUI/JobControl.xaml.cs
+++ b/Engine/Source/Programs/UnsyncUI/UnsyncUI/JobControl.xaml.cs
@@ -3,7 +3,10 @@
 using System;
 using System.Collections.Generic;
 using System.Collections.ObjectModel;
+using System.Collections.Specialized;
+using System.ComponentModel;
 using System.Diagnostics;
+using System.Globalization;
 using System.IO;
 using System.Linq;
 using System.Text;
@@ -12,6 +15,8 @@ using System.Threading;
 using System.Threading.Tasks;
 using System.Windows;
 using System.Windows.Controls;
+using System.Windows.Documents;
+using System.Windows.Media;
 
 namespace UnsyncUI
 {
@@ -20,10 +25,156 @@ namespace UnsyncUI
 	/// </summary>
 	public partial class JobControl : UserControl
     {
+		private static readonly ConsoleColorConverter _colorConverter = new ConsoleColorConverter();
+
+		private ScrollViewer _scrollViewer;
+		private bool _autoScroll = true;
+		private bool _isProgrammaticScroll = false;
+
+		private FlowDocument _logDocument;
+		private LogModel _currentLogModel;
+		private Run _currentRun;
+
         public JobControl()
         {
             InitializeComponent();
+
+			_logDocument = LogRichTextBox.Document;
+			LogRichTextBox.Loaded += OnLogRichTextBoxLoaded;
+			DataContextChanged += OnDataContextChanged;
         }
+
+		private void OnLogRichTextBoxLoaded(object sender, RoutedEventArgs e)
+		{
+			if (_scrollViewer != null)
+				_scrollViewer.ScrollChanged -= OnScrollChanged;
+
+			_scrollViewer = FindScrollViewer(LogRichTextBox);
+			if (_scrollViewer != null)
+				_scrollViewer.ScrollChanged += OnScrollChanged;
+		}
+
+		private static ScrollViewer FindScrollViewer(DependencyObject obj)
+		{
+			for (int i = 0; i < VisualTreeHelper.GetChildrenCount(obj); i++)
+			{
+				var child = VisualTreeHelper.GetChild(obj, i);
+				if (child is ScrollViewer sv)
+					return sv;
+				var result = FindScrollViewer(child);
+				if (result != null)
+					return result;
+			}
+			return null;
+		}
+
+		private static Brush GetBrush(ConsoleColor color)
+			=> (Brush)_colorConverter.Convert(color, typeof(Brush), null, CultureInfo.InvariantCulture);
+
+		private void OnDataContextChanged(object sender, DependencyPropertyChangedEventArgs e)
+		{
+			if (e.OldValue is JobModel oldModel)
+			{
+				oldModel.LogLines.CollectionChanged -= OnLogLinesChanged;
+				if (_currentLogModel != null)
+					_currentLogModel.PropertyChanged -= OnCurrentLogModelChanged;
+			}
+
+			_logDocument.Blocks.Clear();
+			_currentLogModel = null;
+			_currentRun = null;
+
+			if (e.NewValue is JobModel newModel)
+			{
+				// Populate existing lines (e.g. when DataContext switches to a completed job)
+				foreach (var model in newModel.LogLines)
+					AppendLogLine(model);
+
+				newModel.LogLines.CollectionChanged += OnLogLinesChanged;
+				_autoScroll = true;
+				UpdateFollowLink();
+			}
+		}
+
+		private void AppendLogLine(LogModel model)
+		{
+			if (_currentLogModel != null)
+				_currentLogModel.PropertyChanged -= OnCurrentLogModelChanged;
+
+			var run = new Run(model.Text ?? "");
+			run.Foreground = GetBrush(model.Color);
+			var para = new Paragraph(run) { Margin = new Thickness(0) };
+			_logDocument.Blocks.Add(para);
+
+			_currentRun = run;
+			_currentLogModel = model;
+			_currentLogModel.PropertyChanged += OnCurrentLogModelChanged;
+		}
+
+		private void OnCurrentLogModelChanged(object sender, PropertyChangedEventArgs e)
+		{
+			var model = (LogModel)sender;
+			if (_currentRun == null)
+				return;
+
+			if (e.PropertyName == nameof(LogModel.Text))
+				_currentRun.Text = model.Text ?? "";
+			else if (e.PropertyName == nameof(LogModel.Color))
+				_currentRun.Foreground = GetBrush(model.Color);
+		}
+
+		private void OnLogLinesChanged(object sender, NotifyCollectionChangedEventArgs e)
+		{
+			if (e.Action == NotifyCollectionChangedAction.Add)
+			{
+				foreach (LogModel model in e.NewItems)
+					AppendLogLine(model);
+			}
+
+			if (_autoScroll)
+			{
+				_isProgrammaticScroll = true;
+				LogRichTextBox.ScrollToEnd();
+				_isProgrammaticScroll = false;
+			}
+		}
+
+		private void OnScrollChanged(object sender, ScrollChangedEventArgs e)
+		{
+			if (_isProgrammaticScroll)
+				return;
+
+			// ExtentHeightChange != 0 means content grew, not a user scroll
+			if (e.ExtentHeightChange != 0)
+				return;
+
+			_autoScroll = IsAtBottom();
+			UpdateFollowLink();
+		}
+
+		private bool IsAtBottom()
+		{
+			if (_scrollViewer == null)
+				return true;
+
+			return _scrollViewer.VerticalOffset >= _scrollViewer.ScrollableHeight - 2.0;
+		}
+
+		private void UpdateFollowLink()
+		{
+			FollowLink.Visibility = _autoScroll
+				? Visibility.Collapsed
+				: Visibility.Visible;
+		}
+
+		private void FollowLink_Click(object sender, RoutedEventArgs e)
+		{
+			_autoScroll = true;
+			UpdateFollowLink();
+			_isProgrammaticScroll = true;
+			LogRichTextBox.ScrollToEnd();
+			_isProgrammaticScroll = false;
+		}
     }
 
     public sealed class LogModel : BaseModel
diff --git a/Engine/Source/Programs/UnsyncUI/UnsyncUI/UnsyncUI.csproj b/Engine/Source/Programs/UnsyncUI/UnsyncUI/UnsyncUI.csproj
--- a/Engine/Source/Programs/UnsyncUI/UnsyncUI/UnsyncUI.csproj
+++ b/Engine/Source/Programs/UnsyncUI/UnsyncUI/UnsyncUI.csproj
@@ -8,7 +8,7 @@
 		<ApplicationManifest>app.manifest</ApplicationManifest>
 		<IncludeAllContentForSelfExtract>true</IncludeAllContentForSelfExtract>
 		<ApplicationIcon>unsync_48px.ico</ApplicationIcon>
-		<Version>$(VersionPrefix)1.0.26</Version>
+		<Version>$(VersionPrefix)1.0.27</Version>
 	</PropertyGroup>
 
 </Project>
