LauncherTool: Update TargetFramework to net10.0

#rnx
#Synced-CL 52396139

diff --git a/Engine/Source/Programs/UnsyncUI/UnsyncUI/UnsyncUI.csproj b/Engine/Source/Programs/UnsyncUI/UnsyncUI/UnsyncUI.csproj
--- a/Engine/Source/Programs/UnsyncUI/UnsyncUI/UnsyncUI.csproj
+++ b/Engine/Source/Programs/UnsyncUI/UnsyncUI/UnsyncUI.csproj
@@ -2,7 +2,7 @@
 
 	<PropertyGroup>
 		<OutputType>WinExe</OutputType>
-		<TargetFramework>net8.0-windows</TargetFramework>
+		<TargetFramework>net10.0-windows</TargetFramework>
 		<UseWPF>true</UseWPF>
 		<UseWindowsForms>true</UseWindowsForms>
 		<ApplicationManifest>app.manifest</ApplicationManifest>
